Color Management: Adding (advanced) source encodings for textures.

diff --git a/Engine/Source/Developer/TextureBuild/Private/TextureBuildFunction.cpp b/Engine/Source/Developer/TextureBuild/Private/TextureBuildFunction.cpp
--- a/Engine/Source/Developer/TextureBuild/Private/TextureBuildFunction.cpp
+++ b/Engine/Source/Developer/TextureBuild/Private/TextureBuildFunction.cpp
@@ -109,6 +109,7 @@ static FTextureBuildSettings ReadBuildSettingsFromCompactBinary(const FCbObjectV
 	BuildSettings.bVolume = Object["bVolume"].AsBool(BuildSettings.bVolume);
 	BuildSettings.bLongLatSource = Object["bLongLatSource"].AsBool(BuildSettings.bLongLatSource);
 	BuildSettings.bSRGB = Object["bSRGB"].AsBool(BuildSettings.bSRGB);
+	ReadCbField(Object["SourceEncodingOverride"], BuildSettings.SourceEncodingOverride);
 	BuildSettings.bUseLegacyGamma = Object["bUseLegacyGamma"].AsBool(BuildSettings.bUseLegacyGamma);
 	BuildSettings.bPreserveBorder = Object["bPreserveBorder"].AsBool(BuildSettings.bPreserveBorder);
 	BuildSettings.bForceNoAlphaChannel = Object["bForceNoAlphaChannel"].AsBool(BuildSettings.bForceNoAlphaChannel);
diff --git a/Engine/Source/Developer/TextureCompressor/Private/TextureCompressorModule.cpp b/Engine/Source/Developer/TextureCompressor/Private/TextureCompressorModule.cpp
--- a/Engine/Source/Developer/TextureCompressor/Private/TextureCompressorModule.cpp
+++ b/Engine/Source/Developer/TextureCompressor/Private/TextureCompressorModule.cpp
@@ -1206,10 +1206,10 @@ static uint32 ComputeLongLatCubemapExtents(const FImage& SrcImage, const uint32
 	return FMath::Clamp(1U << FMath::FloorLog2(SrcImage.SizeX / 2), 32U, MaxCubemapTextureResolution);
 }
 
-void ITextureCompressorModule::GenerateBaseCubeMipFromLongitudeLatitude2D(FImage* OutMip, const FImage& SrcImage, const uint32 MaxCubemapTextureResolution)
+void ITextureCompressorModule::GenerateBaseCubeMipFromLongitudeLatitude2D(FImage* OutMip, const FImage& SrcImage, const uint32 MaxCubemapTextureResolution, uint8 SourceEncodingOverride)
 {
 	FImage LongLatImage;
-	SrcImage.CopyTo(LongLatImage, ERawImageFormat::RGBA32F, EGammaSpace::Linear);
+	SrcImage.Linearize(SourceEncodingOverride, LongLatImage);
 
 	// TODO_TEXTURE: Expose target size to user.
 	uint32 Extent = ComputeLongLatCubemapExtents(LongLatImage, MaxCubemapTextureResolution);
@@ -2545,7 +2545,6 @@ private:
 		for (int32 MipIndex = StartMip; MipIndex < StartMip + CopyCount; ++MipIndex)
 		{
 			const FImage& Image = SourceMips[MipIndex];
-			ERawImageFormat::Type MipFormat = ERawImageFormat::RGBA32F;
 
 			// create base for the mip chain
 			FImage* Mip = new(OutMipChain) FImage();
@@ -2553,7 +2552,7 @@ private:
 			if (bLongLatCubemap)
 			{
 				// Generate the base mip from the long-lat source image.
-				GenerateBaseCubeMipFromLongitudeLatitude2D(Mip, Image, BuildSettings.MaxTextureResolution);
+				GenerateBaseCubeMipFromLongitudeLatitude2D(Mip, Image, BuildSettings.MaxTextureResolution, BuildSettings.SourceEncodingOverride);
 				break;
 			}
 			else
@@ -2562,9 +2561,7 @@ private:
 				if(BuildSettings.bApplyKernelToTopMip)
 				{
 					FImage Temp;
-					
-					Image.CopyTo(Temp, MipFormat, EGammaSpace::Linear);
-
+					Image.Linearize(BuildSettings.SourceEncodingOverride, Temp);
 					if(BuildSettings.bRenormalizeTopMip)
 					{
 						NormalizeMip(Temp);
@@ -2574,7 +2571,7 @@ private:
 				}
 				else
 				{
-					Image.CopyTo(*Mip, MipFormat, EGammaSpace::Linear);
+					Image.Linearize(BuildSettings.SourceEncodingOverride, *Mip);
 
 					if(BuildSettings.bRenormalizeTopMip)
 					{
diff --git a/Engine/Source/Developer/TextureCompressor/Public/TextureCompressorModule.h b/Engine/Source/Developer/TextureCompressor/Public/TextureCompressorModule.h
--- a/Engine/Source/Developer/TextureCompressor/Public/TextureCompressorModule.h
+++ b/Engine/Source/Developer/TextureCompressor/Public/TextureCompressorModule.h
@@ -105,6 +105,8 @@ struct FTextureBuildSettings
 	uint32 bLongLatSource : 1;
 	/** Whether the texture contains color data in the sRGB colorspace. */
 	uint32 bSRGB : 1;
+	/** Advanced source encoding of the image. */
+	uint8 SourceEncodingOverride;
 	/** Whether the texture should use the legacy gamma space for converting to sRGB */
 	uint32 bUseLegacyGamma : 1;
 	/** Whether the border of the image should be maintained during mipmap generation. */
@@ -199,6 +201,7 @@ struct FTextureBuildSettings
 		, bVolume(false)
 		, bLongLatSource(false)
 		, bSRGB(false)
+		, SourceEncodingOverride(0 /*UE::Color::EEncoding::None*/)
 		, bUseLegacyGamma(false)
 		, bPreserveBorder(false)
 		, bForceNoAlphaChannel(false)
@@ -309,7 +312,7 @@ public:
 	 * @param OutMip - The output mip.
 	 * @param SrcImage - The source longlat image.
 	 */
-	TEXTURECOMPRESSOR_API static void GenerateBaseCubeMipFromLongitudeLatitude2D(FImage* OutMip, const FImage& SrcImage, const uint32 MaxCubemapTextureResolution);
+	TEXTURECOMPRESSOR_API static void GenerateBaseCubeMipFromLongitudeLatitude2D(FImage* OutMip, const FImage& SrcImage, const uint32 MaxCubemapTextureResolution, uint8 SourceEncodingOverride = 0);
 
 
 	/**
diff --git a/Engine/Source/Runtime/ColorManagement/ColorManagement.Build.cs b/Engine/Source/Runtime/ColorManagement/ColorManagement.Build.cs
new file mode 100644
--- /dev/null
+++ b/Engine/Source/Runtime/ColorManagement/ColorManagement.Build.cs
@@ -0,0 +1,21 @@
+// Copyright Epic Games, Inc. All Rights Reserved.
+
+using UnrealBuildTool;
+
+public class ColorManagement : ModuleRules
+{
+	public ColorManagement(ReadOnlyTargetRules Target) : base(Target)
+	{
+		PrivateDependencyModuleNames.AddRange(
+			new string[] {
+				"Core",
+			}
+		);
+
+		PrivateIncludePaths.AddRange(
+			new string[] {
+				"Runtime/ColorManagement/Private"
+			}
+		);
+	}
+}
diff --git a/Engine/Source/Runtime/ColorManagement/Private/Tests/TransferFunctionsTest.cpp b/Engine/Source/Runtime/ColorManagement/Private/Tests/TransferFunctionsTest.cpp
new file mode 100644
--- /dev/null
+++ b/Engine/Source/Runtime/ColorManagement/Private/Tests/TransferFunctionsTest.cpp
@@ -0,0 +1,34 @@
+// Copyright Epic Games, Inc. All Rights Reserved.
+
+#include "CoreTypes.h"
+#include "Containers/UnrealString.h"
+#include "TransferFunctions.h"
+#include "Misc/AutomationTest.h"
+
+#if WITH_DEV_AUTOMATION_TESTS
+
+IMPLEMENT_SIMPLE_AUTOMATION_TEST(FTransferFunctionsTest, "System.ColorManagement.TransferFunctions", EAutomationTestFlags::ApplicationContextMask | EAutomationTestFlags::SmokeFilter)
+
+bool FTransferFunctionsTest::RunTest(const FString& Parameters)
+{
+	using namespace UE::Color;
+
+	const float TestIncrement = 0.05f;
+
+	// Verify that all transfer functions correctly inverse each other.
+
+	for (uint8 EnumValue = static_cast<uint8>(EEncoding::Linear); EnumValue < static_cast<uint8>(EEncoding::Max); EnumValue++)
+	{
+		EEncoding EncodingType = static_cast<EEncoding>(EnumValue);
+
+		for (float TestValue = 0.0f; TestValue <= 1.0f; TestValue += TestIncrement)
+		{
+			float Encoded = UE::Color::Encode(EncodingType, TestValue);
+			float Decoded = UE::Color::Decode(EncodingType, Encoded);
+			TestEqual(TEXT("Transfer function encode followed by decode must match identity"), Decoded, TestValue, KINDA_SMALL_NUMBER);
+		}
+	}
+	return true;
+}
+
+#endif
\ No newline at end of file
diff --git a/Engine/Source/Runtime/ColorManagement/Private/TransferFunctions.cpp b/Engine/Source/Runtime/ColorManagement/Private/TransferFunctions.cpp
new file mode 100644
--- /dev/null
+++ b/Engine/Source/Runtime/ColorManagement/Private/TransferFunctions.cpp
@@ -0,0 +1,126 @@
+// Copyright Epic Games, Inc. All Rights Reserved.
+
+#include "TransferFunctions.h"
+#include "Logging/LogMacros.h"
+
+DEFINE_LOG_CATEGORY_STATIC(LogTransferFunctions, Log, All);
+
+namespace UE { namespace Color {
+
+FORCEINLINE TFunction<float(float)> GetTransferFunction(EEncoding SourceEncoding, bool bIsEncode)
+{
+	switch (SourceEncoding)
+	{
+	case EEncoding::None:        return &Linear;
+	case EEncoding::Linear:      return &Linear;
+	case EEncoding::sRGB:        return bIsEncode ? &EncodeSRGB : &DecodeSRGB;
+	case EEncoding::ST2084:      return bIsEncode ? &EncodeST2084 : &DecodeST2084;
+	case EEncoding::Gamma22:     return bIsEncode ? &EncodeGamma22 : &DecodeGamma22;
+	case EEncoding::BT1886:      return bIsEncode ? &EncodeBT1886 : &DecodeBT1886;
+	case EEncoding::Cineon:      return bIsEncode ? &EncodeCineon : &DecodeCineon;
+	case EEncoding::RedLog:      return bIsEncode ? &EncodeRedLog : &DecodeRedLog;
+	case EEncoding::REDLog3G10:  return bIsEncode ? &EncodeREDLog3G10 : &DecodeREDLog3G10;
+	case EEncoding::SLog1:       return bIsEncode ? &EncodeSLog1 : &DecodeSLog1;
+	case EEncoding::SLog2:       return bIsEncode ? &EncodeSLog2 : &DecodeSLog2;
+	case EEncoding::SLog3:       return bIsEncode ? &EncodeSLog3 : &DecodeSLog3;
+	case EEncoding::AlexaV3LogC: return bIsEncode ? &EncodeArriAlexaV3LogC : &DecodeArriAlexaV3LogC;
+	case EEncoding::CanonLog:    return bIsEncode ? &EncodeCanonLog : &DecodeCanonLog;
+	case EEncoding::ProTune:     return bIsEncode ? &EncodeGoProProTune : &DecodeGoProProTune;
+	case EEncoding::VLog:        return bIsEncode ? &EncodePanasonicVLog : &DecodePanasonicVLog;
+	default: break;
+	}
+
+	UE_LOG(LogTransferFunctions, Warning, TEXT("Failed to find valid transfer function for enum value %d."), SourceEncoding);
+
+	return nullptr;
+}
+
+TFunction<float(float)> GetEncodeFunction(EEncoding SourceEncoding)
+{
+	return GetTransferFunction(SourceEncoding, true);
+}
+
+TFunction<float(float)> GetDecodeFunction(EEncoding SourceEncoding)
+{
+	return GetTransferFunction(SourceEncoding, false);
+}
+
+float Encode(EEncoding SourceEncoding, float Value)
+{
+	if (TFunction<float(float)> TransferFn = GetEncodeFunction(SourceEncoding))
+	{
+		return TransferFn(Value);
+	}
+	else
+	{
+		return Value;
+	}
+}
+
+float Decode(EEncoding SourceEncoding, float Value)
+{
+	if (TFunction<float(float)> TransferFn = GetDecodeFunction(SourceEncoding))
+	{
+		return TransferFn(Value);
+	}
+	else
+	{
+		return Value;
+	}
+}
+
+TFunction<FLinearColor(const FLinearColor&)> GetColorEncodeFunction(EEncoding SourceEncoding)
+{
+	TFunction<float(float)> TransferFn = GetEncodeFunction(SourceEncoding);
+
+	if (!TransferFn)
+	{
+		return nullptr;
+	}
+
+	return [TransferFn](const FLinearColor& Color) -> FLinearColor
+	{
+		return FLinearColor(TransferFn(Color.R), TransferFn(Color.G), TransferFn(Color.B), Color.A);
+	};
+}
+
+TFunction<FLinearColor(const FLinearColor&)> GetColorDecodeFunction(EEncoding SourceEncoding)
+{
+	TFunction<float(float)> TransferFn = GetDecodeFunction(SourceEncoding);
+
+	if (!TransferFn)
+	{
+		return nullptr;
+	}
+
+	return [TransferFn](const FLinearColor& Color) -> FLinearColor
+	{
+		return FLinearColor(TransferFn(Color.R), TransferFn(Color.G), TransferFn(Color.B), Color.A);
+	};
+}
+
+FLinearColor Encode(EEncoding SourceEncoding, const FLinearColor& Color)
+{
+	if (TFunction<FLinearColor(const FLinearColor&)> TransferFn = GetColorEncodeFunction(SourceEncoding))
+	{
+		return TransferFn(Color);
+	}
+	else
+	{
+		return Color;
+	}
+}
+
+FLinearColor Decode(EEncoding SourceEncoding, const FLinearColor& Color)
+{
+	if (TFunction<FLinearColor(const FLinearColor&)> TransferFn = GetColorDecodeFunction(SourceEncoding))
+	{
+		return TransferFn(Color);
+	}
+	else
+	{
+		return Color;
+	}
+}
+
+} } // end namespace UE::Color
diff --git a/Engine/Source/Runtime/ColorManagement/Public/ColorManagementDefines.h b/Engine/Source/Runtime/ColorManagement/Public/ColorManagementDefines.h
new file mode 100644
--- /dev/null
+++ b/Engine/Source/Runtime/ColorManagement/Public/ColorManagementDefines.h
@@ -0,0 +1,34 @@
+// Copyright Epic Games, Inc. All Rights Reserved.
+
+#pragma once
+
+#include "HAL/Platform.h"
+
+namespace UE { namespace Color {
+
+/** List of available encodings/transfer functions.
+* 
+* IMPORTANT NOTE: This list is replicated as a UENUM in TextureDefines.h, and both should always match.
+* Furthermore, make sure to increment TEXTURE_ENCODING_TYPES_VER in TextureDerivedData.cpp upon changes to this list.
+*/
+enum class EEncoding : uint8 {
+	None = 0,
+	Linear = 1,
+	sRGB,
+	ST2084,
+	Gamma22,
+	BT1886,
+	Cineon,
+	RedLog,
+	REDLog3G10,
+	SLog1,
+	SLog2,
+	SLog3,
+	AlexaV3LogC,
+	CanonLog,
+	ProTune,
+	VLog,
+	Max,
+};
+
+} } // end namespace UE::Color
diff --git a/Engine/Source/Runtime/ColorManagement/Public/TransferFunctions.h b/Engine/Source/Runtime/ColorManagement/Public/TransferFunctions.h
new file mode 100644
--- /dev/null
+++ b/Engine/Source/Runtime/ColorManagement/Public/TransferFunctions.h
@@ -0,0 +1,497 @@
+// Copyright Epic Games, Inc. All Rights Reserved.
+
+#pragma once
+
+#include "CoreTypes.h"
+#include "Math/Color.h"
+#include "ColorManagementDefines.h"
+
+namespace UE { namespace Color {
+
+FORCEINLINE float Linear(float Value)
+{
+	return Value;
+}
+
+/**
+* Encode value to sRGB.
+*
+* @return float encoded value.
+*/
+FORCEINLINE float EncodeSRGB(float Value)
+{
+	if (Value <= 0.04045f / 12.92f)
+	{
+		return Value * 12.92f;
+	}
+	else
+	{
+		return  FGenericPlatformMath::Pow(Value, (1.0f / 2.4f)) * 1.055f - 0.055f;
+	}
+}
+
+/**
+* Decode value with an sRGB encoding.
+*
+* @return float decoded value.
+*/
+FORCEINLINE float DecodeSRGB(float Value)
+{
+	if (Value <= 0.04045f)
+	{
+		return Value / 12.92f;
+	}
+	else
+	{
+		return  FGenericPlatformMath::Pow((Value + 0.055f) / 1.055f, 2.4f);
+	}
+}
+
+/**
+* Encode value to SMPTE ST 2084:2014.
+*
+* @return float encoded value.
+*/
+FORCEINLINE float EncodeST2084(float Value)
+{
+	const float Lp = 10000.0f;
+	const float m1 = 2610 / 4096.0f * (1.0f / 4.0f);
+	const float m2 = 2523 / 4096.0f * 128.0f;
+	const float c1 = 3424 / 4096.0f;
+	const float c2 = 2413 / 4096.0f * 32.f;
+	const float c3 = 2392 / 4096.0f * 32.f;
+
+	Value = FGenericPlatformMath::Pow(Value / Lp, m1);
+	return FGenericPlatformMath::Pow((c1 + c2 * Value) / (c3 * Value + 1), m2);
+}
+
+/**
+* Decode value with a SMPTE ST 2084:2014 encoding.
+*
+* @return float decoded value.
+*/
+FORCEINLINE float DecodeST2084(float Value)
+{
+	const float Lp = 10000.0f;
+	const float m1 = 2610 / 4096.0f * (1.0f / 4.0f);
+	const float m2 = 2523 / 4096.0f * 128.0f;
+	const float c1 = 3424 / 4096.0f;
+	const float c2 = 2413 / 4096.0f * 32.f;
+	const float c3 = 2392 / 4096.0f * 32.f;
+
+	const float Vp = FGenericPlatformMath::Pow(Value, 1.0f / m2);
+	Value = FMath::Max(0.0f, Vp - c1);
+	return FGenericPlatformMath::Pow((Value / (c2 - c3 * Vp)), 1.0f / m1) * Lp;
+}
+
+/**
+* Encode value to Gamma 2.2.
+*
+* @return float encoded value.
+*/
+FORCEINLINE float EncodeGamma22(float Value)
+{
+	return FGenericPlatformMath::Pow(Value, 1.0f / 2.2f);
+}
+
+
+/**
+* Decode value with a Gamma 2.2 encoding.
+*
+* @return float decoded value.
+*/
+FORCEINLINE float DecodeGamma22(float Value)
+{
+	return FGenericPlatformMath::Pow(Value, 2.2f);
+}
+
+/**
+* Encode value to ITU-R BT.1886.
+*
+* @return float encoded value.
+*/
+FORCEINLINE float EncodeBT1886(float Value)
+{
+	const float L_B = 0;
+	const float L_W = 1;
+	float Gamma = 2.40f;
+	float GammaInv = 1.0f / Gamma;
+	float N = FGenericPlatformMath::Pow(L_W, GammaInv) - FGenericPlatformMath::Pow(L_B, GammaInv);
+	float A = FGenericPlatformMath::Pow(N, Gamma);
+	float B = FGenericPlatformMath::Pow(L_B, GammaInv) / N;
+	return FGenericPlatformMath::Pow(Value / A, GammaInv) - B;
+}
+
+/**
+* Decode value with an ITU-R BT.1886 encoding.
+*
+* @return float decoded value.
+*/
+FORCEINLINE float DecodeBT1886(float Value)
+{
+	const float L_B = 0;
+	const float L_W = 1;
+	float Gamma = 2.40f;
+	float GammaInv = 1.0f / Gamma;
+	float N = FGenericPlatformMath::Pow(L_W, GammaInv) - FGenericPlatformMath::Pow(L_B, GammaInv);
+	float A = FGenericPlatformMath::Pow(N, Gamma);
+	float B = FGenericPlatformMath::Pow(L_B, GammaInv) / N;
+	return A * FGenericPlatformMath::Pow(FGenericPlatformMath::Max(Value + B, 0.0f), Gamma);
+}
+
+/**
+* Encode value to Cineon.
+*
+* @return float encoded value.
+*/
+FORCEINLINE float EncodeCineon(float Value)
+{
+	const float BlackOffset = FGenericPlatformMath::Pow(10.0f, (95.0f - 685.0f) / 300.0f);
+	return (685.0f + 300.0f * FGenericPlatformMath::LogX(10.0f, Value * (1.0f - BlackOffset) + BlackOffset)) / 1023.0f;
+}
+
+/**
+* Decode value with a Cineon encoding.
+*
+* @return float decoded value.
+*/
+FORCEINLINE float DecodeCineon(float Value)
+{
+	const float BlackOffset = FGenericPlatformMath::Pow(10.0f, (95.0f - 685.0f) / 300.0f);
+	return (FGenericPlatformMath::Pow(10.0f, (1023.0f * Value - 685.0f) / 300.0f) - BlackOffset) / (1.0f - BlackOffset);
+}
+
+/**
+* Encode value to RED Log.
+*
+* @return float encoded value.
+*/
+FORCEINLINE float EncodeRedLog(float Value)
+{
+	const float BlackOffset = FGenericPlatformMath::Pow(10.0f, (0.0f - 1023.0f) / 511.0f);
+	return (1023.0f + 511.0f * FGenericPlatformMath::LogX(10.0f, Value * (1.0f - BlackOffset) + BlackOffset)) / 1023.0f;
+}
+
+/**
+* Decode value with a RED Log encoding.
+*
+* @return float decoded value.
+*/
+FORCEINLINE float DecodeRedLog(float Value)
+{
+	const float BlackOffset = FGenericPlatformMath::Pow(10.0f, (0.0f - 1023.0f) / 511.0f);
+	return (FGenericPlatformMath::Pow(10.0f, (1023.0f * Value - 1023.0f) / 511.0f) - BlackOffset) / (1.0f - BlackOffset);
+}
+
+/**
+* Encode value to RED Log3G10.
+*
+* @return float encoded value.
+*/
+FORCEINLINE float EncodeREDLog3G10(float Value)
+{
+	const float A = 0.224282f;
+	const float B = 155.975327f;
+	const float C = 0.01f;
+	const float G = 15.1927f;
+
+	Value += C;
+
+	if (Value < 0.0f)
+	{
+		return Value * G;
+	}
+	else
+	{
+		return FGenericPlatformMath::Sign(Value) * A * FGenericPlatformMath::LogX(10.0f, B * FGenericPlatformMath::Abs(Value) + 1.0f);
+	}
+}
+
+/**
+* Decode value with a RED Log3G10 encoding.
+*
+* @return float decoded value.
+*/
+FORCEINLINE float DecodeREDLog3G10(float Value)
+{
+	const float A = 0.224282f;
+	const float B = 155.975327f;
+	const float C = 0.01f;
+	const float G = 15.1927f;
+
+	if (Value < 0.0f)
+	{
+		Value /= G;
+	}
+	else
+	{
+		Value = FGenericPlatformMath::Sign(Value) * (FGenericPlatformMath::Pow(10.0f, FGenericPlatformMath::Abs(Value) / A ) - 1.0f) / B;
+	}
+
+	return Value - C;
+}
+
+/**
+* Encode value to Sony S-Log1.
+*
+* @return float encoded value.
+*/
+FORCEINLINE float EncodeSLog1(float Value)
+{
+	Value /= 0.9f;
+	Value = 0.432699f * FGenericPlatformMath::LogX(10.0f, Value + 0.037584f) + 0.616596f + 0.03f;
+	return (Value * 219.0f + 16.0f) * 4.0f / 1023.0f;
+}
+
+/**
+* Decode value with a Sony S-Log1 encoding.
+*
+* @return float decoded value.
+*/
+FORCEINLINE float DecodeSLog1(float Value)
+{
+	Value = ((Value * 1023.f) / 4.0f - 16.0f) / 219.0f;
+	Value = FGenericPlatformMath::Pow(10.0f, (Value - 0.616596f - 0.03f) / 0.432699f) - 0.037584f;
+	return Value * 0.9f;
+}
+
+/**
+* Encode value to Sony S-Log2.
+*
+* @return float encoded value.
+*/
+FORCEINLINE float EncodeSLog2(float Value)
+{
+	if (Value >= 0.0f)
+	{
+		return (64.0f + 876.0f * (0.432699f * FGenericPlatformMath::LogX(10.0f, 155.0f * Value / 197.1f + 0.037584f) + 0.646596f)) / 1023.f;
+	}
+	else
+	{
+		return (64.0f +876.0f * (Value * 3.53881278538813f / 0.9f) + 0.646596f + 0.030001222851889303f) / 1023.f;
+	}
+}
+
+/**
+* Decode value with a Sony S-Log2 encoding.
+*
+* @return float decoded value.
+*/
+FORCEINLINE float DecodeSLog2(float Value)
+{
+	if (Value >= (64.f + 0.030001222851889303f * 876.f) / 1023.f)
+	{
+		return 197.1f * (FGenericPlatformMath::Pow(10.0f, ((Value * 1023.f - 64.f) / 876.f - 0.646596f) / 0.432699f) - 0.037584f) / 155.f;
+	}
+	else
+	{
+		return 0.9f * ((Value * 1023.f - 64.f) / 876.f - 0.030001222851889303f) / 3.53881278538813f;
+	}
+}
+
+/**
+* Encode value to Sony S-Log3.
+*
+* @return float encoded value.
+*/
+FORCEINLINE float EncodeSLog3(float Value)
+{
+	if (Value >= 0.01125000f)
+	{
+		return (420.0f + FGenericPlatformMath::LogX(10.0f, (Value + 0.01f) / 0.19f) * 261.5f) / 1023.0f;
+	}
+	else
+	{
+		return (Value * 76.2102946929f / 0.01125f + 95.0f) / 1023.0f;
+	}
+}
+
+/**
+* Decode value with a Sony S-Log3 encoding.
+*
+* @return float decoded value.
+*/
+FORCEINLINE float DecodeSLog3(float Value)
+{
+	if (Value >= 171.2102946929f / 1023.0f)
+	{
+		return (FGenericPlatformMath::Pow(10.0f, (Value * 1023.0f - 420.f) / 261.5f)) * 0.19f - 0.01f;
+	}
+	else
+	{
+		return (Value * 1023.0f - 95.0f) * 0.01125000f / (171.2102946929f - 95.0f);
+	}
+}
+
+/**
+* Encode value to ARRI Alexa LogC.
+*
+* @return float encoded value.
+*/
+FORCEINLINE float EncodeArriAlexaV3LogC(float Value)
+{
+	const float cut = 0.010591f;
+	const float a = 5.555556f;
+	const float b = 0.052272f;
+	const float c = 0.247190f;
+	const float d = 0.385537f;
+	const float e = 5.367655f;
+	const float f = 0.092809f;
+
+	if (Value > cut)
+	{
+		return c * FGenericPlatformMath::LogX(10.0f, a * Value + b) + d;
+	}
+	else
+	{
+		return e * Value + f;
+	}
+}
+
+/**
+* Decode value with an ARRI Alexa LogC encoding.
+*
+* @return float decoded value.
+*/
+FORCEINLINE float DecodeArriAlexaV3LogC(float Value)
+{
+	const float cut = 0.010591f;
+	const float a = 5.555556f;
+	const float b = 0.052272f;
+	const float c = 0.247190f;
+	const float d = 0.385537f;
+	const float e = 5.367655f;
+	const float f = 0.092809f;
+	
+	if (Value > e * cut + f)
+	{
+		return (FGenericPlatformMath::Pow(10.0f, (Value - d) / c) - b) / a;
+	}
+	else
+	{
+		return (Value - f) / e;
+	}
+}
+
+/**
+* Encode value to Canon Log.
+*
+* @return float encoded value.
+*/
+FORCEINLINE float EncodeCanonLog(float Value)
+{
+	if (Value < 0.0f)
+	{
+		return -(0.529136f * (FGenericPlatformMath::LogX(10.0f, -Value * 10.1596f + 1.0f)) - 0.0730597f);
+	}
+	else
+	{
+		return 0.529136f * FGenericPlatformMath::LogX(10.0f, 10.1596f * Value + 1.0f) + 0.0730597f;
+	}
+}
+
+
+/**
+* Decode value with a Canon Log encoding.
+*
+* @return float decoded value.
+*/
+FORCEINLINE float DecodeCanonLog(float Value)
+{
+	if (Value < 0.0730597f)
+	{
+		return -(FGenericPlatformMath::Pow(10.0f, (0.0730597f - Value) / 0.529136f) - 1.0f) / 10.1596f;
+	}
+	else
+	{
+		return (FGenericPlatformMath::Pow(10.0f, (Value - 0.0730597f) / 0.529136f) - 1.0f) / 10.1596f;
+	}
+}
+
+/**
+* Encode value to GoPro ProTune.
+*
+* @return float encoded value.
+*/
+FORCEINLINE float EncodeGoProProTune(float Value)
+{
+	return FGenericPlatformMath::Loge(Value * 112.f + 1.0f) / FGenericPlatformMath::Loge(113.0f);
+}
+
+/**
+* Decode value with a GoPro ProTune encoding.
+*
+* @return float decoded value.
+*/
+FORCEINLINE float DecodeGoProProTune(float Value)
+{
+	return (FGenericPlatformMath::Pow(113.f, Value) - 1.0f) / 112.f;
+}
+
+/**
+* Encode value to Panasonic V-Log.
+*
+* @return float encoded value.
+*/
+FORCEINLINE float EncodePanasonicVLog(float Value)
+{
+	const float b = 0.00873f;
+	const float c = 0.241514f;
+	const float d = 0.598206f;
+
+	if( Value < 0.01f)
+	{
+		return 5.6f * Value + 0.125f;
+	}
+	else
+	{
+		return c * FGenericPlatformMath::LogX(10.0f, Value + b) + d;
+	}
+}
+
+/**
+* Decode value with a Panasonic V-Log encoding.
+*
+* @return float decoded value.
+*/
+FORCEINLINE float DecodePanasonicVLog(float Value)
+{
+	const float b = 0.00873f;
+	const float c = 0.241514f;
+	const float d = 0.598206f;
+
+	if (Value < 0.181f)
+	{
+		return (Value - 0.125f) / 5.6f;
+	}
+	else
+	{
+		return FGenericPlatformMath::Pow(10.0f, (Value - d) / c) - b;
+	}
+}
+
+/** Get the encode function that matches the encoding type. */
+COLORMANAGEMENT_API TFunction<float(float)> GetEncodeFunction(EEncoding Encoding);
+
+/** Get the decode function that matches the encoding type. */
+COLORMANAGEMENT_API TFunction<float(float)> GetDecodeFunction(EEncoding Encoding);
+
+/** Encode a value based on the specified encoding type. */
+COLORMANAGEMENT_API float Encode(EEncoding Encoding, float Value);
+
+/** Decode a value based on the specified encoding type. */
+COLORMANAGEMENT_API float Decode(EEncoding Encoding, float Value);
+
+/** Get the encode function that matches the encoding type. */
+COLORMANAGEMENT_API TFunction<FLinearColor(const FLinearColor&)> GetColorEncodeFunction(EEncoding Encoding);
+
+/** Get the decode function that matches the encoding type. */
+COLORMANAGEMENT_API TFunction<FLinearColor(const FLinearColor&)> GetColorDecodeFunction(EEncoding Encoding);
+
+/** Encode a color based on the specified encoding type. */
+COLORMANAGEMENT_API FLinearColor Encode(EEncoding Encoding, const FLinearColor& Color);
+
+/** Decode a color based on the specified encoding type. */
+COLORMANAGEMENT_API FLinearColor Decode(EEncoding Encoding, const FLinearColor& Color);
+
+} } // end namespace UE::Color
diff --git a/Engine/Source/Runtime/Engine/Classes/Engine/Texture.h b/Engine/Source/Runtime/Engine/Classes/Engine/Texture.h
--- a/Engine/Source/Runtime/Engine/Classes/Engine/Texture.h
+++ b/Engine/Source/Runtime/Engine/Classes/Engine/Texture.h
@@ -1019,6 +1019,11 @@ public:
 	UPROPERTY(EditAnywhere, BlueprintReadWrite, Category=Texture, meta=(DisplayName="sRGB"), AssetRegistrySearchable)
 	uint8 SRGB:1;
 
+	/** Source encoding of the texture, exposing more options than just sRGB. */
+	UPROPERTY(EditAnywhere, BlueprintReadWrite, Category = Texture, AssetRegistrySearchable, AdvancedDisplay)
+	TEnumAsByte<enum ETextureSourceEncoding> SourceEncodingOverride;
+	
+
 #if WITH_EDITORONLY_DATA
 
 	/** A flag for using the simplified legacy gamma space e.g pow(color,1/2.2) for converting from FColor to FLinearColor, if we're doing sRGB. */
diff --git a/Engine/Source/Runtime/Engine/Classes/Engine/TextureDefines.h b/Engine/Source/Runtime/Engine/Classes/Engine/TextureDefines.h
--- a/Engine/Source/Runtime/Engine/Classes/Engine/TextureDefines.h
+++ b/Engine/Source/Runtime/Engine/Classes/Engine/TextureDefines.h
@@ -295,3 +295,26 @@ enum TextureCompressionSettings
 	TC_EncodedReflectionCapture		UMETA(Hidden),
 	TC_MAX,
 };
+
+/** List of (advanced) texture source encodings, matching the list in `ColorManagementDefines.h`. */
+UENUM()
+enum ETextureSourceEncoding
+{
+	TSE_None = 0 UMETA(DisplayName = "Default", ToolTip = "The source encoding is assumed to match the state of the sRGB checkbox parameter."),
+	TSE_Linear = 1 UMETA(DisplayName = "Linear", ToolTip = "The source encoding is considered linear (before optional sRGB encoding is applied)."),
+	TSE_sRGB = 2 UMETA(DisplayName = "sRGB", ToolTip = "sRGB source encoding to be linearized (before optional sRGB encoding is applied)."),
+	TSE_ST2084 = 3 UMETA(DisplayName = "ST 2084/PQ", ToolTip = "SMPTE ST 2084/PQ source encoding to be linearized (before optional sRGB encoding is applied)."),
+	TSE_Gamma22 = 4 UMETA(DisplayName = "Gamma 2.2", ToolTip = "Gamma 2.2 source encoding to be linearized (before optional sRGB encoding is applied)."),
+	TSE_BT1886 = 5 UMETA(DisplayName = "BT1886", ToolTip = "BT1886 source encoding to be linearized (before optional sRGB encoding is applied)."),
+	TSE_Cineon = 6 UMETA(DisplayName = "Cineon", ToolTip = "Cineon source encoding to be linearized (before optional sRGB encoding is applied)."),
+	TSE_RedLog = 7 UMETA(DisplayName = "RedLog", ToolTip = "RedLog source encoding to be linearized (before optional sRGB encoding is applied)."),
+	TSE_REDLog3G10 = 8 UMETA(DisplayName = "REDLog3G10", ToolTip = "REDLog3G10 source encoding to be linearized (before optional sRGB encoding is applied)."),
+	TSE_SLog1 = 9 UMETA(DisplayName = "SLog1", ToolTip = "Sony SLog1 source encoding to be linearized (before optional sRGB encoding is applied)."),
+	TSE_SLog2 = 10 UMETA(DisplayName = "SLog2", ToolTip = "Sony SLog2 source encoding to be linearized (before optional sRGB encoding is applied)."),
+	TSE_SLog3 = 11 UMETA(DisplayName = "SLog3", ToolTip = "Sony SLog3 source encoding to be linearized (before optional sRGB encoding is applied)."),
+	TSE_AlexaV3LogC = 12 UMETA(DisplayName = "AlexaV3LogC", ToolTip = "ARRI Alexa V3 LogC source encoding to be linearized (before optional sRGB encoding is applied)."),
+	TSE_CanonLog = 13 UMETA(DisplayName = "CanonLog", ToolTip = "Canon Log source encoding to be linearized (before optional sRGB encoding is applied)."),
+	TSE_ProTune = 14 UMETA(DisplayName = "ProTune", ToolTip = "GoPro ProTune source encoding to be linearized (before optional sRGB encoding is applied)."),
+	TSE_VLog = 15 UMETA(DisplayName = "V-Log", ToolTip = "Panasonic V-Log source encoding to be linearized (before optional sRGB encoding is applied)."),
+	TSE_MAX,
+};
diff --git a/Engine/Source/Runtime/Engine/Private/Texture.cpp b/Engine/Source/Runtime/Engine/Private/Texture.cpp
--- a/Engine/Source/Runtime/Engine/Private/Texture.cpp
+++ b/Engine/Source/Runtime/Engine/Private/Texture.cpp
@@ -95,6 +95,7 @@ UTexture::UTexture(const FObjectInitializer& ObjectInitializer)
 		[this](FTextureResource* InTextureResource) { SetResource(InTextureResource); })
 {
 	SRGB = true;
+	SourceEncodingOverride = ETextureSourceEncoding::TSE_None;
 	Filter = TF_Default;
 	MipLoadOptions = ETextureMipLoadOptions::Default;
 #if WITH_EDITORONLY_DATA
diff --git a/Engine/Source/Runtime/Engine/Private/TextureDerivedData.cpp b/Engine/Source/Runtime/Engine/Private/TextureDerivedData.cpp
--- a/Engine/Source/Runtime/Engine/Private/TextureDerivedData.cpp
+++ b/Engine/Source/Runtime/Engine/Private/TextureDerivedData.cpp
@@ -62,6 +62,10 @@
 // This is useful during development, but once large numbers of VT are present in shipped content, it will have the same problem as TEXTURE_DERIVEDDATA_VER
 #define TEXTURE_VT_DERIVEDDATA_VER	TEXT("93F54B64FA524F78BCA4EA2CD30B1E9F")
 
+// Version number for the texture encoding enum types (UE::Color::EEncoding/ETextureSourceEncoding).
+// Make sure to increment upon changes to the list.
+#define TEXTURE_ENCODING_TYPES_VER 1
+
 #if ENABLE_COOK_STATS
 namespace TextureCookStats
 {
@@ -108,6 +112,13 @@ static void SerializeForKey(FArchive& Ar, const FTextureBuildSettings& Settings)
 	TempByte = Settings.bCubemap; Ar << TempByte;
 	TempByte = Settings.bTextureArray; Ar << TempByte;
 	TempByte = Settings.bSRGB ? (Settings.bSRGB | ( Settings.bUseLegacyGamma ? 0 : 0x2 )) : 0; Ar << TempByte;
+
+	if (Settings.SourceEncodingOverride != 0 /*UE::Color::EEncoding::None*/)
+	{
+		TempUint32 = TEXTURE_ENCODING_TYPES_VER; Ar << TempUint32;
+		TempByte = Settings.SourceEncodingOverride; Ar << TempByte;
+	}
+
 	TempByte = Settings.bPreserveBorder; Ar << TempByte;
 	TempByte = Settings.bDitherMipMapAlpha; Ar << TempByte;
 
@@ -472,6 +483,7 @@ static void GetTextureBuildSettings(
 	OutBuildSettings.DiffuseConvolveMipLevel = 0;
 	OutBuildSettings.bLongLatSource = false;
 	OutBuildSettings.bStreamable = false;
+	OutBuildSettings.SourceEncodingOverride = Texture.SourceEncodingOverride.GetValue();
 
 	if (Texture.MaxTextureSize > 0)
 	{
diff --git a/Engine/Source/Runtime/Engine/Private/TextureDerivedDataBuildUtils.cpp b/Engine/Source/Runtime/Engine/Private/TextureDerivedDataBuildUtils.cpp
--- a/Engine/Source/Runtime/Engine/Private/TextureDerivedDataBuildUtils.cpp
+++ b/Engine/Source/Runtime/Engine/Private/TextureDerivedDataBuildUtils.cpp
@@ -111,6 +111,7 @@ static void WriteBuildSettings(FCbWriter& Writer, const FTextureBuildSettings& B
 	WriteCbFieldWithDefault<bool>(Writer, "bVolume", BuildSettings.bVolume, DefaultSettings.bVolume);
 	WriteCbFieldWithDefault<bool>(Writer, "bLongLatSource", BuildSettings.bLongLatSource, DefaultSettings.bLongLatSource);
 	WriteCbFieldWithDefault<bool>(Writer, "bSRGB", BuildSettings.bSRGB, DefaultSettings.bSRGB);
+	WriteCbFieldWithDefault(Writer, "SourceEncodingOverride", BuildSettings.SourceEncodingOverride, DefaultSettings.SourceEncodingOverride);
 	WriteCbFieldWithDefault<bool>(Writer, "bUseLegacyGamma", BuildSettings.bUseLegacyGamma, DefaultSettings.bUseLegacyGamma);
 	WriteCbFieldWithDefault<bool>(Writer, "bPreserveBorder", BuildSettings.bPreserveBorder, DefaultSettings.bPreserveBorder);
 	WriteCbFieldWithDefault<bool>(Writer, "bForceNoAlphaChannel", BuildSettings.bForceNoAlphaChannel, DefaultSettings.bForceNoAlphaChannel);
diff --git a/Engine/Source/Runtime/ImageCore/ImageCore.Build.cs b/Engine/Source/Runtime/ImageCore/ImageCore.Build.cs
--- a/Engine/Source/Runtime/ImageCore/ImageCore.Build.cs
+++ b/Engine/Source/Runtime/ImageCore/ImageCore.Build.cs
@@ -6,6 +6,8 @@ public class ImageCore : ModuleRules
 {
 	public ImageCore(ReadOnlyTargetRules Target) : base(Target)
 	{
+		PrivateDependencyModuleNames.Add("ColorManagement");
+
 		PublicDependencyModuleNames.Add("Core");
 	}
 }
diff --git a/Engine/Source/Runtime/ImageCore/Private/ImageCore.cpp b/Engine/Source/Runtime/ImageCore/Private/ImageCore.cpp
--- a/Engine/Source/Runtime/ImageCore/Private/ImageCore.cpp
+++ b/Engine/Source/Runtime/ImageCore/Private/ImageCore.cpp
@@ -3,7 +3,18 @@
 #include "ImageCore.h"
 #include "Modules/ModuleManager.h"
 #include "Async/ParallelFor.h"
+#include "TransferFunctions.h"
 
+DEFINE_LOG_CATEGORY_STATIC(LogImageCore, Log, All);
+
+static constexpr float MAX_HALF_FLOAT16 = 65504.0f;
+
+FORCEINLINE static void SaturateToHalfFloat(FLinearColor& LinearCol)
+{
+	LinearCol.R = FMath::Clamp(LinearCol.R, -MAX_HALF_FLOAT16, MAX_HALF_FLOAT16);
+	LinearCol.G = FMath::Clamp(LinearCol.G, -MAX_HALF_FLOAT16, MAX_HALF_FLOAT16);
+	LinearCol.B = FMath::Clamp(LinearCol.B, -MAX_HALF_FLOAT16, MAX_HALF_FLOAT16);
+}
 
 IMPLEMENT_MODULE(FDefaultModuleImpl, ImageCore);
 
@@ -442,6 +453,164 @@ void FImage::ResizeTo(FImage& DestImage, int32 DestSizeX, int32 DestSizeY, ERawI
 	}
 }
 
+void FImage::Linearize(uint8 SourceEncoding, FImage& DestImage) const
+{
+	TRACE_CPUPROFILER_EVENT_SCOPE(Linearize);
+
+	DestImage.SizeX = SizeX;
+	DestImage.SizeY = SizeY;
+	DestImage.NumSlices = NumSlices;
+	DestImage.Format = ERawImageFormat::RGBA32F;
+	DestImage.GammaSpace = EGammaSpace::Linear;
+	InitImageStorage(DestImage);
+
+	const FImage& SrcImage = *this;
+	check(SrcImage.SizeX == DestImage.SizeX);
+	check(SrcImage.SizeY == DestImage.SizeY);
+	check(SrcImage.NumSlices == DestImage.NumSlices);
+
+	UE::Color::EEncoding SourceEncodingType = static_cast<UE::Color::EEncoding>(SourceEncoding);
+
+	if (SourceEncodingType == UE::Color::EEncoding::None)
+	{
+		CopyImage(SrcImage, DestImage);
+		return;
+	}
+	else if (SourceEncodingType >= UE::Color::EEncoding::Max)
+	{
+		UE_LOG(LogImageCore, Warning, TEXT("Invalid encoding %d, falling back to linearization using CopyImage."), SourceEncoding);
+		CopyImage(SrcImage, DestImage);
+		return;
+	}
+
+	const bool bDestIsGammaCorrected = DestImage.IsGammaCorrected();
+	const int64 NumTexels = int64(SrcImage.SizeX) * SrcImage.SizeY * SrcImage.NumSlices;
+	const int64 NumJobs = FTaskGraphInterface::Get().GetNumWorkerThreads();
+	int64 TexelsPerJob = NumTexels / NumJobs;
+	if (TexelsPerJob * NumJobs < NumTexels)
+	{
+		++TexelsPerJob;
+	}
+
+	TFunction<FLinearColor(const FLinearColor&)> DecodeFunction = UE::Color::GetColorDecodeFunction(SourceEncodingType);
+	check(DecodeFunction != nullptr);
+
+	// Convert to 32-bit linear floating point.
+	TArrayView64<FLinearColor> DestColors = DestImage.AsRGBA32F();
+	switch (SrcImage.Format)
+	{
+	case ERawImageFormat::G8:
+	{
+		TArrayView64<const uint8> SrcLum = SrcImage.AsG8();
+		for (int64 TexelIndex = 0; TexelIndex < NumTexels; ++TexelIndex)
+		{
+			FColor SrcColor(SrcLum[TexelIndex], SrcLum[TexelIndex], SrcLum[TexelIndex], 255);
+			DestColors[TexelIndex] = SrcColor.ReinterpretAsLinear();
+			DestColors[TexelIndex] = DecodeFunction(DestColors[TexelIndex]);
+			SaturateToHalfFloat(DestColors[TexelIndex]);
+		}
+	}
+	break;
+
+	case ERawImageFormat::G16:
+	{
+		TArrayView64<const uint16> SrcLum = SrcImage.AsG16();
+		for (int64 TexelIndex = 0; TexelIndex < NumTexels; ++TexelIndex)
+		{
+			float Value = FColor::DequantizeUNorm16ToFloat(SrcLum[TexelIndex]);
+			DestColors[TexelIndex] = FLinearColor(Value, Value, Value, 1.0f);
+			DestColors[TexelIndex] = DecodeFunction(DestColors[TexelIndex]);
+			SaturateToHalfFloat(DestColors[TexelIndex]);
+		}
+	}
+	break;
+
+	case ERawImageFormat::BGRA8:
+	{
+		TArrayView64<const FColor> SrcColors = SrcImage.AsBGRA8();
+		EGammaSpace SrcGamma = SrcImage.GammaSpace;
+
+		ParallelFor(NumJobs, [DestColors, SrcColors, TexelsPerJob, NumTexels, SrcGamma, DecodeFunction](int64 JobIndex)
+			{
+				int64 StartIndex = JobIndex * TexelsPerJob;
+				int64 EndIndex = FMath::Min(StartIndex + TexelsPerJob, NumTexels);
+				for (int64 TexelIndex = StartIndex; TexelIndex < EndIndex; ++TexelIndex)
+				{
+					DestColors[TexelIndex] = SrcColors[TexelIndex].ReinterpretAsLinear();
+					DestColors[TexelIndex] = DecodeFunction(DestColors[TexelIndex]);
+					SaturateToHalfFloat(DestColors[TexelIndex]);
+				}
+			});
+	}
+	break;
+
+	case ERawImageFormat::BGRE8:
+	{
+		TArrayView64<const FColor> SrcColors = SrcImage.AsBGRE8();
+		for (int64 TexelIndex = 0; TexelIndex < NumTexels; ++TexelIndex)
+		{
+			DestColors[TexelIndex] = SrcColors[TexelIndex].FromRGBE();
+			DestColors[TexelIndex] = DecodeFunction(DestColors[TexelIndex]);
+			SaturateToHalfFloat(DestColors[TexelIndex]);
+		}
+	}
+	break;
+
+	case ERawImageFormat::RGBA16:
+	{
+		TArrayView64<const uint16> SrcColors = SrcImage.AsRGBA16();
+		for (int64 TexelIndex = 0; TexelIndex < NumTexels; ++TexelIndex)
+		{
+			int64 SrcIndex = TexelIndex * 4;
+			DestColors[TexelIndex] = FLinearColor(
+				SrcColors[SrcIndex + 0] / 65535.0f,
+				SrcColors[SrcIndex + 1] / 65535.0f,
+				SrcColors[SrcIndex + 2] / 65535.0f,
+				SrcColors[SrcIndex + 3] / 65535.0f
+			);
+			DestColors[TexelIndex] = DecodeFunction(DestColors[TexelIndex]);
+			SaturateToHalfFloat(DestColors[TexelIndex]);
+		}
+	}
+	break;
+
+	case ERawImageFormat::RGBA16F:
+	{
+		TArrayView64<const FFloat16Color> SrcColors = SrcImage.AsRGBA16F();
+		for (int64 TexelIndex = 0; TexelIndex < NumTexels; ++TexelIndex)
+		{
+			DestColors[TexelIndex] = SrcColors[TexelIndex].GetFloats();
+			DestColors[TexelIndex] = DecodeFunction(DestColors[TexelIndex]);
+			SaturateToHalfFloat(DestColors[TexelIndex]);
+		}
+	}
+	break;
+
+	case ERawImageFormat::RGBA32F:
+	{
+		TArrayView64<const FLinearColor> SrcColors = SrcImage.AsRGBA32F();
+		for (int64 TexelIndex = 0; TexelIndex < NumTexels; ++TexelIndex)
+		{
+			DestColors[TexelIndex] = DecodeFunction(SrcColors[TexelIndex]);
+			SaturateToHalfFloat(DestColors[TexelIndex]);
+		}
+	}
+	break;
+
+	case ERawImageFormat::R16F:
+	{
+		TArrayView64<const FFloat16> SrcColors = SrcImage.AsR16F();
+		for (int64 TexelIndex = 0; TexelIndex < NumTexels; ++TexelIndex)
+		{
+			DestColors[TexelIndex] = FLinearColor(SrcColors[TexelIndex].GetFloat(), 0, 0, 1);
+			DestColors[TexelIndex] = DecodeFunction(DestColors[TexelIndex]);
+			SaturateToHalfFloat(DestColors[TexelIndex]);
+		}
+	}
+	break;
+	}
+}
+
 int32 FImage::GetBytesPerPixel() const
 {
 	int32 BytesPerPixel = 0;
diff --git a/Engine/Source/Runtime/ImageCore/Public/ImageCore.h b/Engine/Source/Runtime/ImageCore/Public/ImageCore.h
--- a/Engine/Source/Runtime/ImageCore/Public/ImageCore.h
+++ b/Engine/Source/Runtime/ImageCore/Public/ImageCore.h
@@ -105,6 +105,15 @@ public:
 	 */
 	IMAGECORE_API void ResizeTo(FImage& DestImage, int32 DestSizeX, int32 DestSizeY, ERawImageFormat::Type DestFormat, EGammaSpace DestGammaSpace) const;
 
+	/**
+	 * Linearize to a RGBA32F destination image by applying the decoding function that corresponds to the specified source encoding.
+	 * If None, this call will be equivalent to CopyTo(DestImage, ERawImageFormat::RGBA32F, EGammaSpace::Linear).
+	 *
+	 * @param SourceEncoding - Opaque source encoding (matching UE::Color::EEncoding).
+	 * @param DestImage - The destination image.
+	 */
+	IMAGECORE_API void Linearize(uint8 SourceEncoding, FImage& DestImage) const;
+
 	/**
 	 * Gets the number of bytes per pixel.
 	 *
