ColorManagement: Fixes to P3DCI tooltip, improved coordinate conversion functions, color space improvements for later use in output/combine LUTs.

diff --git a/Engine/Source/Runtime/ColorManagement/Private/ColorSpace.cpp b/Engine/Source/Runtime/ColorManagement/Private/ColorSpace.cpp
--- a/Engine/Source/Runtime/ColorManagement/Private/ColorSpace.cpp
+++ b/Engine/Source/Runtime/ColorManagement/Private/ColorSpace.cpp
@@ -17,6 +17,14 @@ void FColorSpace::SetWorking(FColorSpace ColorSpace)
 	WorkingColorSpace = ColorSpace;
 }
 
+static bool IsSRGBChromaticities(const TStaticArray<FCoordinate2d, 4>& Chromaticities, double Tolerance = 1.e-7)
+{
+	return	Chromaticities[0].Equals(FCoordinate2d(0.64, 0.33), Tolerance) &&
+		Chromaticities[1].Equals(FCoordinate2d(0.30, 0.60), Tolerance) &&
+		Chromaticities[2].Equals(FCoordinate2d(0.15, 0.06), Tolerance) &&
+		Chromaticities[3].Equals(FCoordinate2d(0.3127, 0.3290), Tolerance);
+}
+
 FColorSpace::FColorSpace(const FVector2D& InRed, const FVector2D& InGreen, const FVector2D& InBlue, const FVector2D& InWhite)
 {
 	Chromaticities[0] = FCoordinate2d(InRed);
@@ -24,12 +32,15 @@ FColorSpace::FColorSpace(const FVector2D& InRed, const FVector2D& InGreen, const
 	Chromaticities[2] = FCoordinate2d(InBlue);
 	Chromaticities[3] = FCoordinate2d(InWhite);
 
+	bIsSRGB = IsSRGBChromaticities(Chromaticities);
+
 	RgbToXYZ = CalcRgbToXYZ();
 	XYZToRgb = RgbToXYZ.Inverse();
 }
 
 FColorSpace::FColorSpace(EColorSpace ColorSpaceType)
 	: Chromaticities(InPlace, FCoordinate2d(0, 0))
+	, bIsSRGB(ColorSpaceType == EColorSpace::sRGB)
 {
 	const FCoordinate2d D65(0.3127, 0.3290);
 	const FCoordinate2d D60(0.32168, 0.33767);
@@ -128,9 +139,9 @@ FColorSpace::FColorSpace(EColorSpace ColorSpaceType)
 FMatrix44d FColorSpace::CalcRgbToXYZ() const
 {
 	FMatrix44d Mat = FMatrix44d(
-		Chromaticities[0].ToXyY(),
-		Chromaticities[1].ToXyY(),
-		Chromaticities[2].ToXyY(),
+		FVector3d(Chromaticities[0].x, Chromaticities[0].y, 1.0 - Chromaticities[0].x - Chromaticities[0].y),
+		FVector3d(Chromaticities[1].x, Chromaticities[1].y, 1.0 - Chromaticities[1].x - Chromaticities[1].y),
+		FVector3d(Chromaticities[2].x, Chromaticities[2].y, 1.0 - Chromaticities[2].x - Chromaticities[2].y),
 		FVector3d{0,0,0}
 	);
 
@@ -156,7 +167,12 @@ bool FColorSpace::Equals(const FColorSpace& CS, double Tolerance) const
 		Chromaticities[3].Equals(CS.Chromaticities[3], Tolerance);
 }
 
-static FMatrix44d CalcChromaticAdaptionMatrix(FVector4d SourceXYZ, FVector4d TargetXYZ, EChromaticAdaptationMethod Method)
+bool FColorSpace::IsSRGB() const
+{
+	return bIsSRGB;
+}
+
+FMatrix44d FColorSpaceTransform::CalcChromaticAdaptionMatrix(FVector3d SourceXYZ, FVector3d TargetXYZ, EChromaticAdaptationMethod Method)
 {
 	FMatrix44d XyzToRgb;
 
@@ -213,12 +229,12 @@ static FMatrix44d CalcColorSpaceTransformMatrix(const FColorSpace& Src, const FC
 	const FVector3d SrcWhite = Src.GetRgbToXYZ().TransformVector(FVector3d::One());
 	const FVector3d DstWhite = Dst.GetRgbToXYZ().TransformVector(FVector3d::One());
 
-	if (SrcWhite.Equals(DstWhite))
+	if (SrcWhite.Equals(DstWhite, 1.e-7))
 	{
 		return Src.GetRgbToXYZ() * Dst.GetXYZToRgb();
 	}
 
-	FMatrix44d ChromaticAdaptationMat = CalcChromaticAdaptionMatrix(SrcWhite, DstWhite, Method);
+	FMatrix44d ChromaticAdaptationMat = FColorSpaceTransform::CalcChromaticAdaptionMatrix(SrcWhite, DstWhite, Method);
 	return Src.GetRgbToXYZ() * ChromaticAdaptationMat * Dst.GetXYZToRgb();
 }
 
diff --git a/Engine/Source/Runtime/ColorManagement/Private/Tests/ColorManagementTests.cpp b/Engine/Source/Runtime/ColorManagement/Private/Tests/ColorManagementTests.cpp
--- a/Engine/Source/Runtime/ColorManagement/Private/Tests/ColorManagementTests.cpp
+++ b/Engine/Source/Runtime/ColorManagement/Private/Tests/ColorManagementTests.cpp
@@ -11,6 +11,7 @@
 
 #if WITH_DEV_AUTOMATION_TESTS
 
+static bool GPassing;
 static double GdiffDouble = 0.0;
 static int32 GRow = 0;
 static int32 GColumn = 0;
@@ -22,6 +23,7 @@ bool FTransferFunctionsTest::RunTest(const FString& Parameters)
 {
 	using namespace UE::Color;
 
+	GPassing = true;
 	const float TestIncrement = 0.05f;
 
 	// Verify that all transfer functions correctly inverse each other.
@@ -34,10 +36,10 @@ bool FTransferFunctionsTest::RunTest(const FString& Parameters)
 		{
 			float Encoded = UE::Color::Encode(EncodingType, TestValue);
 			float Decoded = UE::Color::Decode(EncodingType, Encoded);
-			TestEqual(TEXT("Transfer function encode followed by decode must match identity"), Decoded, TestValue, KINDA_SMALL_NUMBER);
+			GPassing &= TestEqual(TEXT("Transfer function encode followed by decode must match identity"), Decoded, TestValue, KINDA_SMALL_NUMBER);
 		}
 	}
-	return true;
+	return GPassing;
 }
 
 /**
@@ -72,6 +74,7 @@ FORCENOINLINE void LogTest(const TCHAR* TestName, bool bHasPassed)
 	{
 		UE_LOG(LogUnrealColorManagementTest, Log, TEXT("%s <double>: %s"), bHasPassed ? TEXT("PASSED") : TEXT("FAILED"), TestName);
 		UE_LOG(LogUnrealColorManagementTest, Log, TEXT("Bad(%.8f) at [%i, %i]"), GdiffDouble, GRow, GColumn);
+		GPassing = false;
 	}
 }
 
@@ -219,7 +222,7 @@ static void TestColorSpaceTransforms(UE::Color::EChromaticAdaptationMethod Metho
 			{ -0.00501335146809, -0.0252900718108, 1.03030342328, 0.0 },
 			{ 0,0,0,1 }
 		).GetTransposed();
-		LogTest(TEXT("AP1->Rec709 without chromatic adaptation"), TestMatricesDoubleEqual(Mat0, Mat1, Tolerance));
+		LogTest(TEXT("AP1->Rec709 with Bradford chromatic adaptation"), TestMatricesDoubleEqual(Mat0, Mat1, Tolerance));
 
 		Mat0 = FColorSpaceTransform(Src, FColorSpace(EColorSpace::ACESAP0), Method);
 		Mat1 = FMatrix44d(
@@ -228,11 +231,11 @@ static void TestColorSpaceTransforms(UE::Color::EChromaticAdaptationMethod Metho
 			{ -0.00552588255811, 0.00402521030598, 1.00150067225, 0.0 },
 			{ 0,0,0,1 }
 		).GetTransposed();
-		LogTest(TEXT("AP1->Rec709 without chromatic adaptation"), TestMatricesDoubleEqual(Mat0, Mat1, Tolerance));
+		LogTest(TEXT("AP1->Rec709 with Bradford chromatic adaptation"), TestMatricesDoubleEqual(Mat0, Mat1, Tolerance));
 
 		Mat0 = FColorSpaceTransform(Src, FColorSpace(EColorSpace::ACESAP1), Method);
 		Mat1 = FMatrix44d::Identity;
-		LogTest(TEXT("AP1->Rec709 without chromatic adaptation"), TestMatricesDoubleEqual(Mat0, Mat1, Tolerance));
+		LogTest(TEXT("AP1->Rec709 with Bradford chromatic adaptation"), TestMatricesDoubleEqual(Mat0, Mat1, Tolerance));
 
 		Mat0 = FColorSpaceTransform(Src, FColorSpace(EColorSpace::P3DCI), Method);
 		Mat1 = FMatrix44d(
@@ -241,7 +244,7 @@ static void TestColorSpaceTransforms(UE::Color::EChromaticAdaptationMethod Metho
 			{ -0.00255286167095, -0.0470296027287, 1.0495824644, 0.0 },
 			{ 0,0,0,1 }
 		).GetTransposed();
-		LogTest(TEXT("AP1->Rec709 without chromatic adaptation"), TestMatricesDoubleEqual(Mat0, Mat1, Tolerance));
+		LogTest(TEXT("AP1->Rec709 with Bradford chromatic adaptation"), TestMatricesDoubleEqual(Mat0, Mat1, Tolerance));
 
 		Mat0 = FColorSpaceTransform(Src, FColorSpace(EColorSpace::P3D65), Method);
 		Mat1 = FMatrix44d(
@@ -250,7 +253,7 @@ static void TestColorSpaceTransforms(UE::Color::EChromaticAdaptationMethod Metho
 			{ -0.00215900951357, -0.0454593248373, 1.04761833435, 0.0 },
 			{ 0,0,0,1 }
 		).GetTransposed();
-		LogTest(TEXT("AP1->Rec709 without chromatic adaptation"), TestMatricesDoubleEqual(Mat0, Mat1, Tolerance));
+		LogTest(TEXT("AP1->Rec709 with Bradford chromatic adaptation"), TestMatricesDoubleEqual(Mat0, Mat1, Tolerance));
 	}
 
 }
@@ -260,6 +263,8 @@ bool FColorSpaceTest::RunTest(const FString& Parameters)
 {
 	using namespace UE::Color;
 
+	GPassing = true;
+
 	FColorSpace CS = FColorSpace(EColorSpace::sRGB);
 	FMatrix44d Mat0 = FColorSpaceTransform(CS, CS, EChromaticAdaptationMethod::None);
 	FMatrix44d Mat1 = FMatrix44d::Identity;
@@ -271,7 +276,7 @@ bool FColorSpaceTest::RunTest(const FString& Parameters)
 
 	TestColorSpaceTransforms(EChromaticAdaptationMethod::Bradford);
 
-	return true;
+	return GPassing;
 }
 
 
diff --git a/Engine/Source/Runtime/ColorManagement/Public/ColorSpace.h b/Engine/Source/Runtime/ColorManagement/Public/ColorSpace.h
--- a/Engine/Source/Runtime/ColorManagement/Public/ColorSpace.h
+++ b/Engine/Source/Runtime/ColorManagement/Public/ColorSpace.h
@@ -13,11 +13,11 @@ namespace UE { namespace Color {
 //! At the time of writing, there is no double-precision 2D vector available hence this custom type.
 struct FCoordinate2d
 {
-	/** Coordinate's X component. */
-	double X;
+	/** Coordinate's x component (lowercase to follow the CIE xy-XYZ convention). */
+	double x;
 
-	/** Coordinate's Y component. */
-	double Y;
+	/** Coordinate's y component (lowercase to follow the CIE xy-XYZ convention). */
+	double y;
 
 public:
 
@@ -27,12 +27,12 @@ public:
 	/**
 	* Constructor
 	*
-	* @param InX x coordinate.
-	* @param InY y coordinate.
+	* @param Inx x coordinate.
+	* @param Iny y coordinate.
 	*/
-	explicit FORCEINLINE FCoordinate2d(double InX, double InY)
-		: X(InX)
-		, Y(InY)
+	explicit FORCEINLINE FCoordinate2d(double Inx, double Iny)
+		: x(Inx)
+		, y(Iny)
 	{ }
 
 	/**
@@ -41,8 +41,8 @@ public:
 	* @param Coordinates xy coordinates.
 	*/
 	explicit FORCEINLINE FCoordinate2d(FVector2D Coordinates)
-		: X(Coordinates.X)
-		, Y(Coordinates.Y)
+		: x(Coordinates.X)
+		, y(Coordinates.Y)
 	{ }
 
 	/**
@@ -53,7 +53,7 @@ public:
 	 */
 	FORCEINLINE bool operator==(const FCoordinate2d& V) const
 	{
-		return X == V.X && Y == V.Y;
+		return x == V.x && y == V.x;
 	}
 
 	/**
@@ -63,7 +63,7 @@ public:
 	 */
 	FORCEINLINE bool Equals(const FCoordinate2d& V, float Tolerance) const
 	{
-		return FMath::Abs(X - V.X) <= Tolerance && FMath::Abs(Y - V.Y) <= Tolerance;
+		return FMath::Abs(x - V.x) <= Tolerance && FMath::Abs(y - V.y) <= Tolerance;
 	}
 
 	/**
@@ -75,7 +75,7 @@ public:
 	FORCEINLINE double& operator[](int32 Index)
 	{
 		check(Index >= 0 && Index < 2);
-		return (Index == 0) ? X : Y;
+		return (Index == 0) ? x : y;
 	}
 
 	/**
@@ -87,7 +87,7 @@ public:
 	FORCEINLINE const double& operator[](int32 Index) const
 	{
 		check(Index >= 0 && Index < 2);
-		return (Index == 0) ? X : Y;
+		return (Index == 0) ? x : y;
 	}
 
 	/**
@@ -97,17 +97,32 @@ public:
 	 */
 	FORCEINLINE FVector2D ToVector2D() const
 	{
-		return FVector2D(X, Y);
+		return FVector2D(x, y);
 	}
 
 	/**
-	 * Convert to the CIE xyY colorspace coordinates as an FVector3d.
+	 * Convert to CIE Yxy with a luminance value (default to 1.0).
+	 * 
+	 * @return FVector3d
+	 */
+	FORCEINLINE FVector3d ToYxy(double LuminanceY=1.0) const
+	{
+		return FVector3d(LuminanceY, x, y);
+	}
+
+	/**
+	 * Convert to CIE XYZ tristimulus values with a luminance value (default to 1.0).
 	 * 
 	 * @return FVector3d
 	 */
-	FORCEINLINE FVector3d ToXyY() const
+	FORCEINLINE FVector3d ToXYZ(double LuminanceY = 1.0) const
 	{
-		return FVector3d(X, Y, 1.0 - X - Y);
+		const FVector3d Yxy = ToYxy(LuminanceY);
+		return FVector3d(
+			Yxy[1] * Yxy[0] / FMath::Max(Yxy[2], 1e-10),
+			Yxy[0],
+			(1.0 - Yxy[1] - Yxy[2]) * Yxy[0] / FMath::Max(Yxy[2], 1e-10)
+		);
 	}
 
 	/**
@@ -119,7 +134,7 @@ public:
 	 */
 	friend FArchive& operator<<(FArchive& Ar, FCoordinate2d& C)
 	{
-		return Ar << C.X << C.Y;
+		return Ar << C.x << C.y;
 	}
 
 	bool Serialize(FArchive& Ar)
@@ -200,6 +215,16 @@ public:
 		return Chromaticities[0].ToVector2D();
 	}
 
+	/**
+	* Gets the color space's red chromaticity coordinates.
+	*
+	* @return FCoordinate2d xy coordinates.
+	*/
+	FORCEINLINE FCoordinate2d GetRedChromaticityCoordinate() const
+	{
+		return Chromaticities[0];
+	}
+
 	/**
 	* Gets the color space's green chromaticity coordinates.
 	*
@@ -210,6 +235,16 @@ public:
 		return Chromaticities[1].ToVector2D();
 	}
 
+	/**
+	* Gets the color space's green chromaticity coordinates.
+	*
+	* @return FCoordinate2d xy coordinates.
+	*/
+	FORCEINLINE FCoordinate2d GetGreenChromaticityCoordinate() const
+	{
+		return Chromaticities[1];
+	}
+
 	/**
 	* Gets the color space's blue chromaticity coordinates.
 	*
@@ -220,6 +255,16 @@ public:
 		return Chromaticities[2].ToVector2D();
 	}
 
+	/**
+	* Gets the color space's blue chromaticity coordinates.
+	*
+	* @return FCoordinate2d xy coordinates.
+	*/
+	FORCEINLINE FCoordinate2d GetBlueChromaticityCoordinate() const
+	{
+		return Chromaticities[2];
+	}
+
 	/**
 	* Gets the color space's white point chromaticity coordinates.
 	*
@@ -230,6 +275,16 @@ public:
 		return Chromaticities[3].ToVector2D();
 	}
 
+	/**
+	* Gets the color space's white point chromaticity coordinates.
+	*
+	* @return FCoordinate2d xy coordinates.
+	*/
+	FORCEINLINE FCoordinate2d GetWhiteChromaticityCoordinate() const
+	{
+		return Chromaticities[3];
+	}
+
 	/**
 	* Gets the RGB-to-XYZ conversion matrix.
 	*
@@ -279,7 +334,14 @@ public:
 	 * @param Tolerance Error tolerance.
 	 * @return true if the vectors are equal within tolerance limits, false otherwise.
 	 */
-	bool Equals(const FColorSpace& CS, double Tolerance = SMALL_NUMBER) const;
+	bool Equals(const FColorSpace& CS, double Tolerance = 1.e-7) const;
+
+	/**
+	 * Convenience function to verify if the color space matches the engine's default sRGB chromaticities.
+	 *
+	 * @return true if sRGB.
+	 */
+	bool IsSRGB() const;
 
 private:
 
@@ -291,6 +353,8 @@ private:
 	FMatrix44d RgbToXYZ;
 	FMatrix44d XYZToRgb;
 
+	bool bIsSRGB;
+
 public:
 
 	/**
@@ -337,6 +401,15 @@ struct FColorSpaceTransform : FMatrix44d
 	* @param Color Color to transform.
 	*/
 	COLORMANAGEMENT_API FLinearColor Apply(const FLinearColor& Color) const;
+
+	/**
+	* Calculate the chromatic adaptation matrix using the specific method.
+	*
+	* @param SourceXYZ Source color in XYZ space.
+	* @param TargetXYZ Target color in XYZ space.
+	* @param Method Adaptation method (None, Bradford, CAT02).
+	*/
+	COLORMANAGEMENT_API static FMatrix44d CalcChromaticAdaptionMatrix(FVector3d SourceXYZ, FVector3d TargetXYZ, EChromaticAdaptationMethod Method);
 };
 
 } }  // end namespace UE::Color
diff --git a/Engine/Source/Runtime/Engine/Classes/Engine/RendererSettings.h b/Engine/Source/Runtime/Engine/Classes/Engine/RendererSettings.h
--- a/Engine/Source/Runtime/Engine/Classes/Engine/RendererSettings.h
+++ b/Engine/Source/Runtime/Engine/Classes/Engine/RendererSettings.h
@@ -224,7 +224,7 @@ namespace EWorkingColorSpace
 		Rec2020 = 2 UMETA(DisplayName = "Rec2020", ToolTip = "Rec2020 (BT.2020) primaries with D65 white point."),
 		ACESAP0 = 3 UMETA(DIsplayName = "ACES AP0", ToolTip = "ACES AP0 wide gamut primaries, with D60 white point."),
 		ACESAP1 = 4 UMETA(DIsplayName = "ACES AP1 / ACEScg", ToolTip = "ACES AP1 / ACEScg wide gamut primaries, with D60 white point."),
-		P3DCI = 5 UMETA(DisplayName = "P3DCI", ToolTip = "P3 (Theater) primaries, with D65 white point."),
+		P3DCI = 5 UMETA(DisplayName = "P3DCI", ToolTip = "P3 (Theater) primaries, with DCI Calibration white point."),
 		P3D65 = 6 UMETA(DisplayName = "P3D65", ToolTip = "P3 (Display) primaries, with D65 white point."),
 		Custom = 7 UMETA(DisplayName = "Custom", ToolTip = "User defined color space and white point."),
 	};
diff --git a/Engine/Source/Runtime/Engine/Classes/Engine/TextureDefines.h b/Engine/Source/Runtime/Engine/Classes/Engine/TextureDefines.h
--- a/Engine/Source/Runtime/Engine/Classes/Engine/TextureDefines.h
+++ b/Engine/Source/Runtime/Engine/Classes/Engine/TextureDefines.h
@@ -332,7 +332,7 @@ enum class ETextureColorSpace : uint8
 	TCS_Rec2020				= 2 UMETA(DisplayName = "Rec2020", ToolTip = "Rec2020 (BT.2020) primaries with D65 white point."),
 	TCS_ACESAP0				= 3 UMETA(DIsplayName = "ACES AP0", ToolTip = "ACES AP0 wide gamut primaries, with D60 white point."),
 	TCS_ACESAP1				= 4 UMETA(DIsplayName = "ACES AP1 / ACEScg", ToolTip = "ACES AP1 / ACEScg wide gamut primaries, with D60 white point."),
-	TCS_P3DCI				= 5 UMETA(DisplayName = "P3DCI", ToolTip = "P3 (Theater) primaries, with D65 white point."),
+	TCS_P3DCI				= 5 UMETA(DisplayName = "P3DCI", ToolTip = "P3 (Theater) primaries, with DCI Calibration white point."),
 	TCS_P3D65				= 6 UMETA(DisplayName = "P3D65", ToolTip = "P3 (Display) primaries, with D65 white point."),
 	TCS_REDWideGamut		= 7 UMETA(DisplayName = "RED Wide Gamut", ToolTip = "RED Wide Gamut primaries, with D65 white point."),
 	TCS_SonySGamut3			= 8 UMETA(DisplayName = "Sony S-Gamut3", ToolTip = "Sony S-Gamut/S-Gamut3 primaries, with D65 white point."),
