ColorManagement: Removing custom coordinate type to replace it with the now available FVector2d.

diff --git a/Engine/Source/Developer/TextureCompressor/Private/TextureCompressorModule.cpp b/Engine/Source/Developer/TextureCompressor/Private/TextureCompressorModule.cpp
--- a/Engine/Source/Developer/TextureCompressor/Private/TextureCompressorModule.cpp
+++ b/Engine/Source/Developer/TextureCompressor/Private/TextureCompressorModule.cpp
@@ -3034,10 +3034,10 @@ private:
 			if (BuildSettings.bHasColorSpaceDefinition)
 			{
 				Mip->TransformToWorkingColorSpace(
-					FVector2D(BuildSettings.RedChromaticityCoordinate),
-					FVector2D(BuildSettings.GreenChromaticityCoordinate),
-					FVector2D(BuildSettings.BlueChromaticityCoordinate),
-					FVector2D(BuildSettings.WhiteChromaticityCoordinate),
+					FVector2d(BuildSettings.RedChromaticityCoordinate),
+					FVector2d(BuildSettings.GreenChromaticityCoordinate),
+					FVector2d(BuildSettings.BlueChromaticityCoordinate),
+					FVector2d(BuildSettings.WhiteChromaticityCoordinate),
 					static_cast<UE::Color::EChromaticAdaptationMethod>(BuildSettings.ChromaticAdaptationMethod));
 			}
 
diff --git a/Engine/Source/Runtime/ColorManagement/Private/ColorSpace.cpp b/Engine/Source/Runtime/ColorManagement/Private/ColorSpace.cpp
--- a/Engine/Source/Runtime/ColorManagement/Private/ColorSpace.cpp
+++ b/Engine/Source/Runtime/ColorManagement/Private/ColorSpace.cpp
@@ -17,20 +17,20 @@ void FColorSpace::SetWorking(FColorSpace ColorSpace)
 	WorkingColorSpace = ColorSpace;
 }
 
-static bool IsSRGBChromaticities(const TStaticArray<FCoordinate2d, 4>& Chromaticities, double Tolerance = 1.e-7)
+static bool IsSRGBChromaticities(const TStaticArray<FVector2d, 4>& Chromaticities, double Tolerance = 1.e-7)
 {
-	return	Chromaticities[0].Equals(FCoordinate2d(0.64, 0.33), Tolerance) &&
-		Chromaticities[1].Equals(FCoordinate2d(0.30, 0.60), Tolerance) &&
-		Chromaticities[2].Equals(FCoordinate2d(0.15, 0.06), Tolerance) &&
-		Chromaticities[3].Equals(FCoordinate2d(0.3127, 0.3290), Tolerance);
+	return	Chromaticities[0].Equals(FVector2d(0.64, 0.33), Tolerance) &&
+		Chromaticities[1].Equals(FVector2d(0.30, 0.60), Tolerance) &&
+		Chromaticities[2].Equals(FVector2d(0.15, 0.06), Tolerance) &&
+		Chromaticities[3].Equals(FVector2d(0.3127, 0.3290), Tolerance);
 }
 
-FColorSpace::FColorSpace(const FVector2D& InRed, const FVector2D& InGreen, const FVector2D& InBlue, const FVector2D& InWhite)
+FColorSpace::FColorSpace(const FVector2d& InRed, const FVector2d& InGreen, const FVector2d& InBlue, const FVector2d& InWhite)
 {
-	Chromaticities[0] = FCoordinate2d(InRed);
-	Chromaticities[1] = FCoordinate2d(InGreen);
-	Chromaticities[2] = FCoordinate2d(InBlue);
-	Chromaticities[3] = FCoordinate2d(InWhite);
+	Chromaticities[0] = InRed;
+	Chromaticities[1] = InGreen;
+	Chromaticities[2] = InBlue;
+	Chromaticities[3] = InWhite;
 
 	bIsSRGB = IsSRGBChromaticities(Chromaticities);
 
@@ -39,92 +39,92 @@ FColorSpace::FColorSpace(const FVector2D& InRed, const FVector2D& InGreen, const
 }
 
 FColorSpace::FColorSpace(EColorSpace ColorSpaceType)
-	: Chromaticities(InPlace, FCoordinate2d(0, 0))
+	: Chromaticities(InPlace, FVector2d::Zero())
 	, bIsSRGB(ColorSpaceType == EColorSpace::sRGB)
 {
-	const FCoordinate2d D65(0.3127, 0.3290);
-	const FCoordinate2d D60(0.32168, 0.33767);
+	const FVector2d D65(0.3127, 0.3290);
+	const FVector2d D60(0.32168, 0.33767);
 
 	switch (ColorSpaceType)
 	{
 	case EColorSpace::None:
 		break;
 	case EColorSpace::sRGB:
-		Chromaticities[0] = FCoordinate2d(0.64, 0.33);
-		Chromaticities[1] = FCoordinate2d(0.30, 0.60);
-		Chromaticities[2] = FCoordinate2d(0.15, 0.06);
+		Chromaticities[0] = FVector2d(0.64, 0.33);
+		Chromaticities[1] = FVector2d(0.30, 0.60);
+		Chromaticities[2] = FVector2d(0.15, 0.06);
 		Chromaticities[3] = D65;
 		break;
 	case EColorSpace::Rec2020:
-		Chromaticities[0] = FCoordinate2d(0.708, 0.292);
-		Chromaticities[1] = FCoordinate2d(0.170, 0.797);
-		Chromaticities[2] = FCoordinate2d(0.131, 0.046);
+		Chromaticities[0] = FVector2d(0.708, 0.292);
+		Chromaticities[1] = FVector2d(0.170, 0.797);
+		Chromaticities[2] = FVector2d(0.131, 0.046);
 		Chromaticities[3] = D65;
 		break;
 	case EColorSpace::ACESAP0:
-		Chromaticities[0] = FCoordinate2d(0.7347, 0.2653);
-		Chromaticities[1] = FCoordinate2d(0.0000, 1.0000);
-		Chromaticities[2] = FCoordinate2d(0.0001, -0.0770);
+		Chromaticities[0] = FVector2d(0.7347, 0.2653);
+		Chromaticities[1] = FVector2d(0.0000, 1.0000);
+		Chromaticities[2] = FVector2d(0.0001, -0.0770);
 		Chromaticities[3] = D60;
 		break;
 	case EColorSpace::ACESAP1:
-		Chromaticities[0] = FCoordinate2d(0.713, 0.293);
-		Chromaticities[1] = FCoordinate2d(0.165, 0.830);
-		Chromaticities[2] = FCoordinate2d(0.128, 0.044);
+		Chromaticities[0] = FVector2d(0.713, 0.293);
+		Chromaticities[1] = FVector2d(0.165, 0.830);
+		Chromaticities[2] = FVector2d(0.128, 0.044);
 		Chromaticities[3] = D60;
 		break;
 	case EColorSpace::P3DCI:
-		Chromaticities[0] = FCoordinate2d(0.680, 0.320);
-		Chromaticities[1] = FCoordinate2d(0.265, 0.690);
-		Chromaticities[2] = FCoordinate2d(0.150, 0.060);
-		Chromaticities[3] = FCoordinate2d(0.314, 0.351);
+		Chromaticities[0] = FVector2d(0.680, 0.320);
+		Chromaticities[1] = FVector2d(0.265, 0.690);
+		Chromaticities[2] = FVector2d(0.150, 0.060);
+		Chromaticities[3] = FVector2d(0.314, 0.351);
 		break;
 	case EColorSpace::P3D65:
-		Chromaticities[0] = FCoordinate2d(0.680, 0.320);
-		Chromaticities[1] = FCoordinate2d(0.265, 0.690);
-		Chromaticities[2] = FCoordinate2d(0.150, 0.060);
+		Chromaticities[0] = FVector2d(0.680, 0.320);
+		Chromaticities[1] = FVector2d(0.265, 0.690);
+		Chromaticities[2] = FVector2d(0.150, 0.060);
 		Chromaticities[3] = D65;
 		break;
 	case EColorSpace::REDWideGamut:
-		Chromaticities[0] = FCoordinate2d(0.780308, 0.304253);
-		Chromaticities[1] = FCoordinate2d(0.121595, 1.493994);
-		Chromaticities[2] = FCoordinate2d(0.095612, -0.084589);
+		Chromaticities[0] = FVector2d(0.780308, 0.304253);
+		Chromaticities[1] = FVector2d(0.121595, 1.493994);
+		Chromaticities[2] = FVector2d(0.095612, -0.084589);
 		Chromaticities[3] = D65;
 		break;
 	case EColorSpace::SonySGamut3:
-		Chromaticities[0] = FCoordinate2d(0.730, 0.280);
-		Chromaticities[1] = FCoordinate2d(0.140, 0.855);
-		Chromaticities[2] = FCoordinate2d(0.100, -0.050);
+		Chromaticities[0] = FVector2d(0.730, 0.280);
+		Chromaticities[1] = FVector2d(0.140, 0.855);
+		Chromaticities[2] = FVector2d(0.100, -0.050);
 		Chromaticities[3] = D65;
 		break;
 	case EColorSpace::SonySGamut3Cine:
-		Chromaticities[0] = FCoordinate2d(0.766, 0.275);
-		Chromaticities[1] = FCoordinate2d(0.225, 0.800);
-		Chromaticities[2] = FCoordinate2d(0.089, -0.087);
+		Chromaticities[0] = FVector2d(0.766, 0.275);
+		Chromaticities[1] = FVector2d(0.225, 0.800);
+		Chromaticities[2] = FVector2d(0.089, -0.087);
 		Chromaticities[3] = D65;
 		break;
 	case EColorSpace::AlexaWideGamut:
-		Chromaticities[0] = FCoordinate2d(0.684, 0.313);
-		Chromaticities[1] = FCoordinate2d(0.221, 0.848);
-		Chromaticities[2] = FCoordinate2d(0.0861, -0.1020);
+		Chromaticities[0] = FVector2d(0.684, 0.313);
+		Chromaticities[1] = FVector2d(0.221, 0.848);
+		Chromaticities[2] = FVector2d(0.0861, -0.1020);
 		Chromaticities[3] = D65;
 		break;
 	case EColorSpace::CanonCinemaGamut:
-		Chromaticities[0] = FCoordinate2d(0.740, 0.270);
-		Chromaticities[1] = FCoordinate2d(0.170, 1.140);
-		Chromaticities[2] = FCoordinate2d(0.080, -0.100);
+		Chromaticities[0] = FVector2d(0.740, 0.270);
+		Chromaticities[1] = FVector2d(0.170, 1.140);
+		Chromaticities[2] = FVector2d(0.080, -0.100);
 		Chromaticities[3] = D65;
 		break;
 	case EColorSpace::GoProProtuneNative:
-		Chromaticities[0] = FCoordinate2d(0.698448, 0.193026);
-		Chromaticities[1] = FCoordinate2d(0.329555, 1.024597);
-		Chromaticities[2] = FCoordinate2d(0.108443, -0.034679);
+		Chromaticities[0] = FVector2d(0.698448, 0.193026);
+		Chromaticities[1] = FVector2d(0.329555, 1.024597);
+		Chromaticities[2] = FVector2d(0.108443, -0.034679);
 		Chromaticities[3] = D65;
 		break;
 	case EColorSpace::PanasonicVGamut:
-		Chromaticities[0] = FCoordinate2d(0.730, 0.280);
-		Chromaticities[1] = FCoordinate2d(0.165, 0.840);
-		Chromaticities[2] = FCoordinate2d(0.100, -0.030);
+		Chromaticities[0] = FVector2d(0.730, 0.280);
+		Chromaticities[1] = FVector2d(0.165, 0.840);
+		Chromaticities[2] = FVector2d(0.100, -0.030);
 		Chromaticities[3] = D65;
 		break;
 	default:
@@ -139,14 +139,14 @@ FColorSpace::FColorSpace(EColorSpace ColorSpaceType)
 FMatrix44d FColorSpace::CalcRgbToXYZ() const
 {
 	FMatrix44d Mat = FMatrix44d(
-		FVector3d(Chromaticities[0].x, Chromaticities[0].y, 1.0 - Chromaticities[0].x - Chromaticities[0].y),
-		FVector3d(Chromaticities[1].x, Chromaticities[1].y, 1.0 - Chromaticities[1].x - Chromaticities[1].y),
-		FVector3d(Chromaticities[2].x, Chromaticities[2].y, 1.0 - Chromaticities[2].x - Chromaticities[2].y),
+		FVector3d(Chromaticities[0].X, Chromaticities[0].Y, 1.0 - Chromaticities[0].X - Chromaticities[0].Y),
+		FVector3d(Chromaticities[1].X, Chromaticities[1].Y, 1.0 - Chromaticities[1].X - Chromaticities[1].Y),
+		FVector3d(Chromaticities[2].X, Chromaticities[2].Y, 1.0 - Chromaticities[2].X - Chromaticities[2].Y),
 		FVector3d{0,0,0}
 	);
 
 	FMatrix44d Inverse = Mat.Inverse();
-	const FCoordinate2d& White = Chromaticities[3];
+	const FVector2d& White = Chromaticities[3];
 	FVector3d WhiteXYZ = FVector3d(White[0] / White[1], 1.0, (1.0f - White[0] - White[1]) / White[1]);
 	FVector3d Scale = Inverse.TransformVector(WhiteXYZ);
 
diff --git a/Engine/Source/Runtime/ColorManagement/Public/ColorSpace.h b/Engine/Source/Runtime/ColorManagement/Public/ColorSpace.h
--- a/Engine/Source/Runtime/ColorManagement/Public/ColorSpace.h
+++ b/Engine/Source/Runtime/ColorManagement/Public/ColorSpace.h
@@ -10,139 +10,50 @@
 
 namespace UE { namespace Color {
 
-//! At the time of writing, there is no double-precision 2D vector available hence this custom type.
-struct FCoordinate2d
+/**
+ * Convert coordinate to CIE Yxy with a luminance value.
+ *
+ * @return FVector3d
+ */
+FORCEINLINE FVector3d ToYxy(double LuminanceY, const FVector2d& Coordinate)
 {
-	/** Coordinate's x component (lowercase to follow the CIE xy-XYZ convention). */
-	double x;
-
-	/** Coordinate's y component (lowercase to follow the CIE xy-XYZ convention). */
-	double y;
-
-public:
-
-	/** Constructor */
-	FORCEINLINE FCoordinate2d() { }
-
-	/**
-	* Constructor
-	*
-	* @param Inx x coordinate.
-	* @param Iny y coordinate.
-	*/
-	explicit FORCEINLINE FCoordinate2d(double Inx, double Iny)
-		: x(Inx)
-		, y(Iny)
-	{ }
-
-	/**
-	* Constructor
-	*
-	* @param Coordinates xy coordinates.
-	*/
-	explicit FORCEINLINE FCoordinate2d(FVector2D Coordinates)
-		: x(Coordinates.X)
-		, y(Coordinates.Y)
-	{ }
-
-	/**
-	 * Compares these coordinates against another pair for equality.
-	 *
-	 * @param V The FCoordinate2d to compare against.
-	 * @return true if the two coordinates are equal, otherwise false.
-	 */
-	FORCEINLINE bool operator==(const FCoordinate2d& V) const
-	{
-		return x == V.x && y == V.x;
-	}
-
-	/**
-	 * Convert to the CIE xyY colorspace coordinates as an FVector3d.
-	 *
-	 * @return true if equal
-	 */
-	FORCEINLINE bool Equals(const FCoordinate2d& V, float Tolerance) const
-	{
-		return FMath::Abs(x - V.x) <= Tolerance && FMath::Abs(y - V.y) <= Tolerance;
-	}
-
-	/**
-	 * Gets specific component of the coordinates.
-	 *
-	 * @param Index the index of coordinate component
-	 * @return reference to component.
-	 */
-	FORCEINLINE double& operator[](int32 Index)
-	{
-		check(Index >= 0 && Index < 2);
-		return (Index == 0) ? x : y;
-	}
-
-	/**
-	 * Gets specific component of the coordinate.
-	 *
-	 * @param Index the index of coordinate component
-	 * @return copy of component value.
-	 */
-	FORCEINLINE const double& operator[](int32 Index) const
-	{
-		check(Index >= 0 && Index < 2);
-		return (Index == 0) ? x : y;
-	}
-
-	/**
-	 * Convert to FVector2D.
-	 * 
-	 * @return FVector2D
-	 */
-	FORCEINLINE FVector2D ToVector2D() const
-	{
-		return FVector2D(x, y);
-	}
-
-	/**
-	 * Convert to CIE Yxy with a luminance value (default to 1.0).
-	 * 
-	 * @return FVector3d
-	 */
-	FORCEINLINE FVector3d ToYxy(double LuminanceY=1.0) const
-	{
-		return FVector3d(LuminanceY, x, y);
-	}
-
-	/**
-	 * Convert to CIE XYZ tristimulus values with a luminance value (default to 1.0).
-	 * 
-	 * @return FVector3d
-	 */
-	FORCEINLINE FVector3d ToXYZ(double LuminanceY = 1.0) const
-	{
-		const FVector3d Yxy = ToYxy(LuminanceY);
-		return FVector3d(
-			Yxy[1] * Yxy[0] / FMath::Max(Yxy[2], 1e-10),
-			Yxy[0],
-			(1.0 - Yxy[1] - Yxy[2]) * Yxy[0] / FMath::Max(Yxy[2], 1e-10)
-		);
-	}
-
-	/**
-	 * Serialize a coordinate.
-	 *
-	 * @param Ar Serialization archive.
-	 * @param C Coordinate2d being serialized.
-	 * @return Reference to Archive after serialization.
-	 */
-	friend FArchive& operator<<(FArchive& Ar, FCoordinate2d& C)
-	{
-		return Ar << C.x << C.y;
-	}
-
-	bool Serialize(FArchive& Ar)
-	{
-		Ar << *this;
-		return true;
-	}
-};
+	return FVector3d(LuminanceY, Coordinate.X, Coordinate.Y);
+}
+
+/**
+ * Convert coordinate to CIE Yxy with a default luminance of 1.0.
+ *
+ * @return FVector3d
+ */
+FORCEINLINE FVector3d ToYxy(const FVector2d& Coordinate)
+{
+	return ToYxy(1.0, Coordinate);
+}
+
+/**
+ * Convert coordinate to CIE XYZ tristimulus values with a luminance value.
+ * 
+ * @return FVector3d
+ */
+FORCEINLINE FVector3d ToXYZ(double LuminanceY, const FVector2d& Coordinate)
+{
+	const FVector3d Yxy = ToYxy(LuminanceY, Coordinate);
+	return FVector3d(
+		Yxy[1] * Yxy[0] / FMath::Max(Yxy[2], 1e-10),
+		Yxy[0],
+		(1.0 - Yxy[1] - Yxy[2]) * Yxy[0] / FMath::Max(Yxy[2], 1e-10)
+	);
+}
+
+/**
+ * Convert coordinate to CIE XYZ tristimulus values with a default luminance of 1.0.
+ *
+ * @return FVector3d
+ */
+FORCEINLINE FVector3d ToXYZ(const FVector2d& Coordinate)
+{
+	return ToXYZ(1.0, Coordinate);
+}
 
 
 /** Color space definition as 4 chromaticity coordinates, in double precision internally. */
@@ -175,7 +86,7 @@ public:
 	* @param InBlue Chromaticity 2D coordinates for the blue color.
 	* @param InWhite Chromaticity 2D coordinates for the white point.
 	*/
-	explicit FColorSpace(const FVector2D& InRed, const FVector2D& InGreen, const FVector2D& InBlue, const FVector2D& InWhite);
+	explicit FColorSpace(const FVector2d& InRed, const FVector2d& InGreen, const FVector2d& InBlue, const FVector2d& InWhite);
 
 	/**
 	* Constructor
@@ -190,37 +101,28 @@ public:
 	FColorSpace& operator=(const FColorSpace&) = default;
 
 	/**
-	* Getter for the color space chromaticity coordinates as FVector2D.
+	* Getter for the color space chromaticity coordinates.
 	* 
-	* @param OutRed FVector2D for the red color chromaticity coordinate.
-	* @param OutGreen FVector2D for the green color chromaticity coordinate.
-	* @param OutBlue FVector2D for the blue color chromaticity coordinate.
-	* @param OutWhite FVector2D for the white color chromaticity coordinate.
-	*/
-	FORCEINLINE void GetChromaticities(FVector2D& OutRed, FVector2D& OutGreen, FVector2D& OutBlue, FVector2D& OutWhite) const
-	{
-		OutRed		= Chromaticities[0].ToVector2D();
-		OutGreen	= Chromaticities[1].ToVector2D();
-		OutBlue		= Chromaticities[2].ToVector2D();
-		OutWhite	= Chromaticities[3].ToVector2D();
-	}
-
-	/**
-	* Gets the color space's red chromaticity coordinates.
-	*
-	* @return FVector2D xy coordinates.
+	* @param OutRed FVector2d for the red color chromaticity coordinate.
+	* @param OutGreen FVector2d for the green color chromaticity coordinate.
+	* @param OutBlue FVector2d for the blue color chromaticity coordinate.
+	* @param OutWhite FVector2d for the white color chromaticity coordinate.
 	*/
-	FORCEINLINE FVector2D GetRedChromaticity() const
+	template<typename T>
+	FORCEINLINE void GetChromaticities(UE::Math::TVector2<T>& OutRed, UE::Math::TVector2<T>& OutGreen, UE::Math::TVector2<T>& OutBlue, UE::Math::TVector2<T>& OutWhite) const
 	{
-		return Chromaticities[0].ToVector2D();
+		OutRed		= Chromaticities[0];
+		OutGreen	= Chromaticities[1];
+		OutBlue		= Chromaticities[2];
+		OutWhite	= Chromaticities[3];
 	}
 
 	/**
 	* Gets the color space's red chromaticity coordinates.
 	*
-	* @return FCoordinate2d xy coordinates.
+	* @return FVector2d xy coordinates.
 	*/
-	FORCEINLINE FCoordinate2d GetRedChromaticityCoordinate() const
+	FORCEINLINE const FVector2d& GetRedChromaticity() const
 	{
 		return Chromaticities[0];
 	}
@@ -228,19 +130,9 @@ public:
 	/**
 	* Gets the color space's green chromaticity coordinates.
 	*
-	* @return FVector2D xy coordinates.
+	* @return FVector2d xy coordinates.
 	*/
-	FORCEINLINE FVector2D GetGreenChromaticity() const
-	{
-		return Chromaticities[1].ToVector2D();
-	}
-
-	/**
-	* Gets the color space's green chromaticity coordinates.
-	*
-	* @return FCoordinate2d xy coordinates.
-	*/
-	FORCEINLINE FCoordinate2d GetGreenChromaticityCoordinate() const
+	FORCEINLINE const FVector2d& GetGreenChromaticity() const
 	{
 		return Chromaticities[1];
 	}
@@ -248,19 +140,9 @@ public:
 	/**
 	* Gets the color space's blue chromaticity coordinates.
 	*
-	* @return FVector2D xy coordinates.
+	* @return FVector2d xy coordinates.
 	*/
-	FORCEINLINE FVector2D GetBlueChromaticity() const
-	{
-		return Chromaticities[2].ToVector2D();
-	}
-
-	/**
-	* Gets the color space's blue chromaticity coordinates.
-	*
-	* @return FCoordinate2d xy coordinates.
-	*/
-	FORCEINLINE FCoordinate2d GetBlueChromaticityCoordinate() const
+	FORCEINLINE const FVector2d& GetBlueChromaticity() const
 	{
 		return Chromaticities[2];
 	}
@@ -268,19 +150,9 @@ public:
 	/**
 	* Gets the color space's white point chromaticity coordinates.
 	*
-	* @return FVector2D xy coordinates.
-	*/
-	FORCEINLINE FVector2D GetWhiteChromaticity() const
-	{
-		return Chromaticities[3].ToVector2D();
-	}
-
-	/**
-	* Gets the color space's white point chromaticity coordinates.
-	*
-	* @return FCoordinate2d xy coordinates.
+	* @return FVector2d xy coordinates.
 	*/
-	FORCEINLINE FCoordinate2d GetWhiteChromaticityCoordinate() const
+	FORCEINLINE const FVector2d& GetWhiteChromaticity() const
 	{
 		return Chromaticities[3];
 	}
@@ -348,7 +220,7 @@ private:
 	FMatrix44d CalcRgbToXYZ() const;
 
 	/** Red, green, blue, white chromaticities, in order. */
-	TStaticArray<FCoordinate2d, 4> Chromaticities;
+	TStaticArray<FVector2d, 4> Chromaticities;
 
 	FMatrix44d RgbToXYZ;
 	FMatrix44d XYZToRgb;
diff --git a/Engine/Source/Runtime/ImageCore/Private/ImageCore.cpp b/Engine/Source/Runtime/ImageCore/Private/ImageCore.cpp
--- a/Engine/Source/Runtime/ImageCore/Private/ImageCore.cpp
+++ b/Engine/Source/Runtime/ImageCore/Private/ImageCore.cpp
@@ -535,7 +535,7 @@ static void CopyImage(const FImage& SrcImage, FImage& DestImage)
 	}
 }
 
-void FImage::TransformToWorkingColorSpace(const FVector2D& SourceRedChromaticity, const FVector2D& SourceGreenChromaticity, const FVector2D& SourceBlueChromaticity, const FVector2D& SourceWhiteChromaticity, UE::Color::EChromaticAdaptationMethod Method, double EqualityTolerance)
+void FImage::TransformToWorkingColorSpace(const FVector2d& SourceRedChromaticity, const FVector2d& SourceGreenChromaticity, const FVector2d& SourceBlueChromaticity, const FVector2d& SourceWhiteChromaticity, UE::Color::EChromaticAdaptationMethod Method, double EqualityTolerance)
 {
 	TRACE_CPUPROFILER_EVENT_SCOPE(Texture.TransformToWorkingColorSpace);
 
diff --git a/Engine/Source/Runtime/ImageCore/Public/ImageCore.h b/Engine/Source/Runtime/ImageCore/Public/ImageCore.h
--- a/Engine/Source/Runtime/ImageCore/Public/ImageCore.h
+++ b/Engine/Source/Runtime/ImageCore/Public/ImageCore.h
@@ -130,7 +130,7 @@ public:
 	 * @param Method - The chromatic adapation method.
 	 * @param EqualityTolerance - The tolerance for the source and working color space chromaticities to be considered equal, bypassing the transform.
 	 */
-	IMAGECORE_API void TransformToWorkingColorSpace(const FVector2D& SourceRedChromaticity, const FVector2D& SourceGreenChromaticity, const FVector2D& SourceBlueChromaticity, const FVector2D& SourceWhiteChromaticity, UE::Color::EChromaticAdaptationMethod Method, double EqualityTolerance = 1.e-7);
+	IMAGECORE_API void TransformToWorkingColorSpace(const FVector2d& SourceRedChromaticity, const FVector2d& SourceGreenChromaticity, const FVector2d& SourceBlueChromaticity, const FVector2d& SourceWhiteChromaticity, UE::Color::EChromaticAdaptationMethod Method, double EqualityTolerance = 1.e-7);
 
 	/**
 	 * Gets the number of bytes per pixel.
