ColorManagement: Allow external access to the standard (in use) white point coordinates, also added convenience matrix transpose functions for shaders.

diff --git a/Engine/Source/Runtime/ColorManagement/Private/ColorSpace.cpp b/Engine/Source/Runtime/ColorManagement/Private/ColorSpace.cpp
--- a/Engine/Source/Runtime/ColorManagement/Private/ColorSpace.cpp
+++ b/Engine/Source/Runtime/ColorManagement/Private/ColorSpace.cpp
@@ -42,9 +42,6 @@ FColorSpace::FColorSpace(EColorSpace ColorSpaceType)
 	: Chromaticities(InPlace, FVector2d::Zero())
 	, bIsSRGB(ColorSpaceType == EColorSpace::sRGB)
 {
-	const FVector2d D65(0.3127, 0.3290);
-	const FVector2d D60(0.32168, 0.33767);
-
 	switch (ColorSpaceType)
 	{
 	case EColorSpace::None:
@@ -54,82 +51,82 @@ FColorSpace::FColorSpace(EColorSpace ColorSpaceType)
 		Chromaticities[0] = FVector2d(0.64, 0.33);
 		Chromaticities[1] = FVector2d(0.30, 0.60);
 		Chromaticities[2] = FVector2d(0.15, 0.06);
-		Chromaticities[3] = D65;
+		Chromaticities[3] = GetWhitePoint(EWhitePoint::CIE1931_D65);
 		break;
 	case EColorSpace::Rec2020:
 		Chromaticities[0] = FVector2d(0.708, 0.292);
 		Chromaticities[1] = FVector2d(0.170, 0.797);
 		Chromaticities[2] = FVector2d(0.131, 0.046);
-		Chromaticities[3] = D65;
+		Chromaticities[3] = GetWhitePoint(EWhitePoint::CIE1931_D65);
 		break;
 	case EColorSpace::ACESAP0:
 		Chromaticities[0] = FVector2d(0.7347, 0.2653);
 		Chromaticities[1] = FVector2d(0.0000, 1.0000);
 		Chromaticities[2] = FVector2d(0.0001, -0.0770);
-		Chromaticities[3] = D60;
+		Chromaticities[3] = GetWhitePoint(EWhitePoint::ACES_D60);
 		break;
 	case EColorSpace::ACESAP1:
 		Chromaticities[0] = FVector2d(0.713, 0.293);
 		Chromaticities[1] = FVector2d(0.165, 0.830);
 		Chromaticities[2] = FVector2d(0.128, 0.044);
-		Chromaticities[3] = D60;
+		Chromaticities[3] = GetWhitePoint(EWhitePoint::ACES_D60);
 		break;
 	case EColorSpace::P3DCI:
 		Chromaticities[0] = FVector2d(0.680, 0.320);
 		Chromaticities[1] = FVector2d(0.265, 0.690);
 		Chromaticities[2] = FVector2d(0.150, 0.060);
-		Chromaticities[3] = FVector2d(0.314, 0.351);
+		Chromaticities[3] = GetWhitePoint(EWhitePoint::DCI_CalibrationWhite);
 		break;
 	case EColorSpace::P3D65:
 		Chromaticities[0] = FVector2d(0.680, 0.320);
 		Chromaticities[1] = FVector2d(0.265, 0.690);
 		Chromaticities[2] = FVector2d(0.150, 0.060);
-		Chromaticities[3] = D65;
+		Chromaticities[3] = GetWhitePoint(EWhitePoint::CIE1931_D65);
 		break;
 	case EColorSpace::REDWideGamut:
 		Chromaticities[0] = FVector2d(0.780308, 0.304253);
 		Chromaticities[1] = FVector2d(0.121595, 1.493994);
 		Chromaticities[2] = FVector2d(0.095612, -0.084589);
-		Chromaticities[3] = D65;
+		Chromaticities[3] = GetWhitePoint(EWhitePoint::CIE1931_D65);
 		break;
 	case EColorSpace::SonySGamut3:
 		Chromaticities[0] = FVector2d(0.730, 0.280);
 		Chromaticities[1] = FVector2d(0.140, 0.855);
 		Chromaticities[2] = FVector2d(0.100, -0.050);
-		Chromaticities[3] = D65;
+		Chromaticities[3] = GetWhitePoint(EWhitePoint::CIE1931_D65);
 		break;
 	case EColorSpace::SonySGamut3Cine:
 		Chromaticities[0] = FVector2d(0.766, 0.275);
 		Chromaticities[1] = FVector2d(0.225, 0.800);
 		Chromaticities[2] = FVector2d(0.089, -0.087);
-		Chromaticities[3] = D65;
+		Chromaticities[3] = GetWhitePoint(EWhitePoint::CIE1931_D65);
 		break;
 	case EColorSpace::AlexaWideGamut:
 		Chromaticities[0] = FVector2d(0.684, 0.313);
 		Chromaticities[1] = FVector2d(0.221, 0.848);
 		Chromaticities[2] = FVector2d(0.0861, -0.1020);
-		Chromaticities[3] = D65;
+		Chromaticities[3] = GetWhitePoint(EWhitePoint::CIE1931_D65);
 		break;
 	case EColorSpace::CanonCinemaGamut:
 		Chromaticities[0] = FVector2d(0.740, 0.270);
 		Chromaticities[1] = FVector2d(0.170, 1.140);
 		Chromaticities[2] = FVector2d(0.080, -0.100);
-		Chromaticities[3] = D65;
+		Chromaticities[3] = GetWhitePoint(EWhitePoint::CIE1931_D65);
 		break;
 	case EColorSpace::GoProProtuneNative:
 		Chromaticities[0] = FVector2d(0.698448, 0.193026);
 		Chromaticities[1] = FVector2d(0.329555, 1.024597);
 		Chromaticities[2] = FVector2d(0.108443, -0.034679);
-		Chromaticities[3] = D65;
+		Chromaticities[3] = GetWhitePoint(EWhitePoint::CIE1931_D65);
 		break;
 	case EColorSpace::PanasonicVGamut:
 		Chromaticities[0] = FVector2d(0.730, 0.280);
 		Chromaticities[1] = FVector2d(0.165, 0.840);
 		Chromaticities[2] = FVector2d(0.100, -0.030);
-		Chromaticities[3] = D65;
+		Chromaticities[3] = GetWhitePoint(EWhitePoint::CIE1931_D65);
 		break;
 	default:
-		check(false);
+		checkNoEntry();
 		break;
 	}
 
diff --git a/Engine/Source/Runtime/ColorManagement/Public/ColorManagementDefines.h b/Engine/Source/Runtime/ColorManagement/Public/ColorManagementDefines.h
--- a/Engine/Source/Runtime/ColorManagement/Public/ColorManagementDefines.h
+++ b/Engine/Source/Runtime/ColorManagement/Public/ColorManagementDefines.h
@@ -72,4 +72,12 @@ enum class EChromaticAdaptationMethod : uint8
 	CAT02 = 2,
 };
 
+/** List of standard white points. */
+enum class EWhitePoint : uint8
+{
+	CIE1931_D65 = 0,
+	ACES_D60 = 1,
+	DCI_CalibrationWhite,
+};
+
 } } // end namespace UE::Color
diff --git a/Engine/Source/Runtime/ColorManagement/Public/ColorSpace.h b/Engine/Source/Runtime/ColorManagement/Public/ColorSpace.h
--- a/Engine/Source/Runtime/ColorManagement/Public/ColorSpace.h
+++ b/Engine/Source/Runtime/ColorManagement/Public/ColorSpace.h
@@ -55,6 +55,45 @@ FORCEINLINE FVector3d ToXYZ(const FVector2d& Coordinate)
 	return ToXYZ(1.0, Coordinate);
 }
 
+/**
+ * Convenience function to get the transposed matrix, i.e. for pre-multiplied shader matrices.
+ *
+ * @return TMatrix<T>
+ */
+template<typename T>
+FORCEINLINE UE::Math::TMatrix<T> Transpose(const FMatrix44d& Transform)
+{
+	if constexpr (TIsSame<T, double>::Value)
+	{
+		return Transform.GetTransposed();
+	}
+
+	return UE::Math::TMatrix<T>(Transform).GetTransposed();
+}
+
+/**
+ * Get standard white point coordinates.
+ *
+ * @param InWhitePoint White point type.
+ * @return FVector2d Chromaticity coordinates.
+ */
+FORCEINLINE FVector2d GetWhitePoint(EWhitePoint InWhitePoint)
+{
+	switch (InWhitePoint)
+	{
+	case UE::Color::EWhitePoint::CIE1931_D65:
+		return FVector2d(0.3127, 0.3290);
+	case UE::Color::EWhitePoint::ACES_D60:
+		return FVector2d(0.32168, 0.33767);
+	case UE::Color::EWhitePoint::DCI_CalibrationWhite:
+		return FVector2d(0.314, 0.351);
+
+	default:
+		checkNoEntry();
+		return FVector2d();
+	}
+}
+
 
 /** Color space definition as 4 chromaticity coordinates, in double precision internally. */
 class COLORMANAGEMENT_API FColorSpace
