ColorManagement: Prevent copy on SetWorking call with move assignment.

#rnx

diff --git a/Engine/Source/Runtime/ColorManagement/Private/ColorSpace.cpp b/Engine/Source/Runtime/ColorManagement/Private/ColorSpace.cpp
--- a/Engine/Source/Runtime/ColorManagement/Private/ColorSpace.cpp
+++ b/Engine/Source/Runtime/ColorManagement/Private/ColorSpace.cpp
@@ -14,7 +14,7 @@ const FColorSpace& FColorSpace::GetWorking()
 
 void FColorSpace::SetWorking(FColorSpace ColorSpace)
 {
-	WorkingColorSpace = ColorSpace;
+	WorkingColorSpace = MoveTemp(ColorSpace);
 }
 
 static bool IsSRGBChromaticities(const TStaticArray<FVector2d, 4>& Chromaticities, double Tolerance = 1.e-7)
