ColorManagement: Bradford chromatic adaption method as the default transform argument.

#rnx

diff --git a/Engine/Source/Runtime/ColorManagement/Public/ColorSpace.h b/Engine/Source/Runtime/ColorManagement/Public/ColorSpace.h
--- a/Engine/Source/Runtime/ColorManagement/Public/ColorSpace.h
+++ b/Engine/Source/Runtime/ColorManagement/Public/ColorSpace.h
@@ -303,7 +303,7 @@ struct FColorSpaceTransform : FMatrix44d
 	* @param Target Target color space.
 	* @param Method Chromatic adapation method.
 	*/
-	COLORMANAGEMENT_API explicit FColorSpaceTransform(const FColorSpace& Src, const FColorSpace& Dst, EChromaticAdaptationMethod Method);
+	COLORMANAGEMENT_API explicit FColorSpaceTransform(const FColorSpace& Src, const FColorSpace& Dst, EChromaticAdaptationMethod Method = EChromaticAdaptationMethod::Bradford);
 
 	/**
 	* Constructor: create a color space transformation matrix from a raw matrix.
@@ -326,7 +326,7 @@ struct FColorSpaceTransform : FMatrix44d
 	* @param TargetXYZ Target color in XYZ space.
 	* @param Method Adaptation method (None, Bradford, CAT02).
 	*/
-	COLORMANAGEMENT_API static FMatrix44d CalcChromaticAdaptionMatrix(FVector3d SourceXYZ, FVector3d TargetXYZ, EChromaticAdaptationMethod Method);
+	COLORMANAGEMENT_API static FMatrix44d CalcChromaticAdaptionMatrix(FVector3d SourceXYZ, FVector3d TargetXYZ, EChromaticAdaptationMethod Method = EChromaticAdaptationMethod::Bradford);
 };
 
 } }  // end namespace UE::Color
diff --git a/Engine/Source/Runtime/Engine/Private/SceneManagement.cpp b/Engine/Source/Runtime/Engine/Private/SceneManagement.cpp
--- a/Engine/Source/Runtime/Engine/Private/SceneManagement.cpp
+++ b/Engine/Source/Runtime/Engine/Private/SceneManagement.cpp
@@ -79,10 +79,10 @@ void FDefaultWorkingColorSpaceUniformBuffer::Update(const UE::Color::FColorSpace
 	Parameters.ToXYZ = Transpose<float>(InColorSpace.GetRgbToXYZ());
 	Parameters.FromXYZ = Transpose<float>(InColorSpace.GetXYZToRgb());
 
-	Parameters.ToAP1 = Transpose<float>(FColorSpaceTransform(InColorSpace, FColorSpace(EColorSpace::ACESAP1), EChromaticAdaptationMethod::Bradford));
+	Parameters.ToAP1 = Transpose<float>(FColorSpaceTransform(InColorSpace, FColorSpace(EColorSpace::ACESAP1)));
 	Parameters.FromAP1 = Parameters.ToAP1.Inverse();
 	
-	Parameters.ToAP0 = Transpose<float>(FColorSpaceTransform(InColorSpace, FColorSpace(EColorSpace::ACESAP0), EChromaticAdaptationMethod::Bradford));
+	Parameters.ToAP0 = Transpose<float>(FColorSpaceTransform(InColorSpace, FColorSpace(EColorSpace::ACESAP0)));
 
 	Parameters.bIsSRGB = InColorSpace.IsSRGB();
 
