ColorManagement: Adding ANSI E1.54 PLASA color space definition.

This is mainly for color accuracy of calibrated DMX light fixtures.

diff --git a/Engine/Source/Runtime/ColorManagement/Private/ColorSpace.cpp b/Engine/Source/Runtime/ColorManagement/Private/ColorSpace.cpp
--- a/Engine/Source/Runtime/ColorManagement/Private/ColorSpace.cpp
+++ b/Engine/Source/Runtime/ColorManagement/Private/ColorSpace.cpp
@@ -125,6 +125,12 @@ FColorSpace::FColorSpace(EColorSpace ColorSpaceType)
 		Chromaticities[2] = FVector2d(0.100, -0.030);
 		Chromaticities[3] = GetWhitePoint(EWhitePoint::CIE1931_D65);
 		break;
+	case EColorSpace::PLASA_E1_54:
+		Chromaticities[0] = FVector2d(0.7347, 0.2653);
+		Chromaticities[1] = FVector2d(0.1596, 0.8404);
+		Chromaticities[2] = FVector2d(0.0366, 0.0001);
+		Chromaticities[3] = FVector2d(0.4254, 0.4044);
+		break;
 	default:
 		checkNoEntry();
 		break;
diff --git a/Engine/Source/Runtime/ColorManagement/Public/ColorManagementDefines.h b/Engine/Source/Runtime/ColorManagement/Public/ColorManagementDefines.h
--- a/Engine/Source/Runtime/ColorManagement/Public/ColorManagementDefines.h
+++ b/Engine/Source/Runtime/ColorManagement/Public/ColorManagementDefines.h
@@ -40,24 +40,25 @@ enum class EEncoding : uint8
 
 /** List of available color spaces. (Increment COLORSPACE_VER upon breaking changes to the list.)
 * 
-* NOTE: This list is replicated as a UENUM in TextureDefines.h, and both should always match.
+* NOTE: This list is partially replicated as a UENUM in TextureDefines.h: any type exposed to textures should match the enum value below.
 */
 enum class EColorSpace : uint8
 {
 	None = 0,
 	sRGB = 1,
-	Rec2020,
-	ACESAP0,
-	ACESAP1,
-	P3DCI,
-	P3D65,
-	REDWideGamut,
-	SonySGamut3,
-	SonySGamut3Cine,
-	AlexaWideGamut,
-	CanonCinemaGamut,
-	GoProProtuneNative,
-	PanasonicVGamut,
+	Rec2020 = 2,
+	ACESAP0 = 3,
+	ACESAP1 = 4,
+	P3DCI = 5,
+	P3D65 = 6,
+	REDWideGamut = 7,
+	SonySGamut3 = 8,
+	SonySGamut3Cine = 9,
+	AlexaWideGamut = 10,
+	CanonCinemaGamut = 11,
+	GoProProtuneNative = 12,
+	PanasonicVGamut = 13,
+	PLASA_E1_54 = 14,
 };
 
 
