ColorManagement: Ensure that the light component color temperature is based on the working color space.

diff --git a/Engine/Source/Editor/UnrealEd/Private/Lightmass/Lightmass.cpp b/Engine/Source/Editor/UnrealEd/Private/Lightmass/Lightmass.cpp
--- a/Engine/Source/Editor/UnrealEd/Private/Lightmass/Lightmass.cpp
+++ b/Engine/Source/Editor/UnrealEd/Private/Lightmass/Lightmass.cpp
@@ -214,7 +214,7 @@ void Copy( const ULightComponent* In, Lightmass::FLightData& Out )
 
 	if( In->bUseTemperature )
 	{
-		Out.Color *= FLinearColor::MakeFromColorTemperature(In->Temperature);
+		Out.Color *= In->GetColorTemperature();
 	}
 }
 
diff --git a/Engine/Source/Runtime/ColorManagement/Private/ColorSpace.cpp b/Engine/Source/Runtime/ColorManagement/Private/ColorSpace.cpp
--- a/Engine/Source/Runtime/ColorManagement/Private/ColorSpace.cpp
+++ b/Engine/Source/Runtime/ColorManagement/Private/ColorSpace.cpp
@@ -176,6 +176,25 @@ bool FColorSpace::IsSRGB() const
 	return bIsSRGB;
 }
 
+FLinearColor FColorSpace::MakeFromColorTemperature(float Temp) const
+{
+	Temp = FMath::Clamp(Temp, 1000.0f, 15000.0f);
+
+	// Approximate Planckian locus in CIE 1960 UCS
+	float u = (0.860117757f + 1.54118254e-4f * Temp + 1.28641212e-7f * Temp * Temp) / (1.0f + 8.42420235e-4f * Temp + 7.08145163e-7f * Temp * Temp);
+	float v = (0.317398726f + 4.22806245e-5f * Temp + 4.20481691e-8f * Temp * Temp) / (1.0f - 2.89741816e-5f * Temp + 1.61456053e-7f * Temp * Temp);
+
+	float x = 3.0f * u / (2.0f * u - 8.0f * v + 4.0f);
+	float y = 2.0f * v / (2.0f * u - 8.0f * v + 4.0f);
+	float z = 1.0f - x - y;
+
+	FVector3d XYZ = FVector3d(1.0 / y * x, 1.0, 1.0 / y * z);
+	FVector4d RGB = XYZToRgb.TransformVector(XYZ);
+
+	return FLinearColor((float)RGB.X, (float)RGB.Y, (float)RGB.Z);
+}
+
+
 FMatrix44d FColorSpaceTransform::CalcChromaticAdaptionMatrix(FVector3d SourceXYZ, FVector3d TargetXYZ, EChromaticAdaptationMethod Method)
 {
 	FMatrix44d XyzToRgb;
diff --git a/Engine/Source/Runtime/ColorManagement/Public/ColorSpace.h b/Engine/Source/Runtime/ColorManagement/Public/ColorSpace.h
--- a/Engine/Source/Runtime/ColorManagement/Public/ColorSpace.h
+++ b/Engine/Source/Runtime/ColorManagement/Public/ColorSpace.h
@@ -260,6 +260,11 @@ public:
 	 */
 	bool IsSRGB() const;
 
+	/**
+	* Converts temperature in Kelvins of a black body radiator to an RGB color in the current space.
+	*/
+	FLinearColor MakeFromColorTemperature(float Temp) const;
+
 private:
 
 	FMatrix44d CalcRgbToXYZ() const;
diff --git a/Engine/Source/Runtime/Engine/Classes/Components/LightComponent.h b/Engine/Source/Runtime/Engine/Classes/Components/LightComponent.h
--- a/Engine/Source/Runtime/Engine/Classes/Components/LightComponent.h
+++ b/Engine/Source/Runtime/Engine/Classes/Components/LightComponent.h
@@ -475,6 +475,9 @@ public:
 
 	FLinearColor GetColoredLightBrightness() const;
 
+	/** Get the color temperature in the working color space. */
+	FLinearColor GetColorTemperature() const;
+
 	/** 
 	 * Called when property is modified by InterpPropertyTracks
 	 *
diff --git a/Engine/Source/Runtime/Engine/Private/Components/LightComponent.cpp b/Engine/Source/Runtime/Engine/Private/Components/LightComponent.cpp
--- a/Engine/Source/Runtime/Engine/Private/Components/LightComponent.cpp
+++ b/Engine/Source/Runtime/Engine/Private/Components/LightComponent.cpp
@@ -28,6 +28,7 @@
 #include "Components/DirectionalLightComponent.h"
 #include "Components/BillboardComponent.h"
 #include "ComponentRecreateRenderStateContext.h"
+#include "ColorSpace.h"
 
 #include UE_INLINE_GENERATED_CPP_BY_NAME(LightComponent)
 
@@ -1497,12 +1498,17 @@ FLinearColor ULightComponent::GetColoredLightBrightness() const
 	FLinearColor Energy = FLinearColor(LightColor) * LightBrightness;
 	if (bUseTemperature)
 	{
-		Energy *= FLinearColor::MakeFromColorTemperature(Temperature);
+		Energy *= GetColorTemperature();
 	}
 
 	return Energy;
 }
 
+FLinearColor ULightComponent::GetColorTemperature() const
+{
+	return UE::Color::FColorSpace::GetWorking().MakeFromColorTemperature(Temperature);
+}
+
 UMaterialInterface* ULightComponent::GetMaterial(int32 ElementIndex) const
 {
 	if (ElementIndex == 0)
diff --git a/Engine/Source/Runtime/Renderer/Private/RendererScene.cpp b/Engine/Source/Runtime/Renderer/Private/RendererScene.cpp
--- a/Engine/Source/Runtime/Renderer/Private/RendererScene.cpp
+++ b/Engine/Source/Runtime/Renderer/Private/RendererScene.cpp
@@ -3800,15 +3800,10 @@ void FScene::UpdateLightColorAndBrightness(ULightComponent* Light)
 		};
 
 		FUpdateLightColorParameters NewParameters;
-		NewParameters.NewColor = FLinearColor(Light->LightColor) * Light->ComputeLightBrightness();
+		NewParameters.NewColor = Light->GetColoredLightBrightness();
 		NewParameters.NewIndirectLightingScale = Light->IndirectLightingIntensity;
 		NewParameters.NewVolumetricScatteringIntensity = Light->VolumetricScatteringIntensity;
 
-		if( Light->bUseTemperature )
-		{
-			 NewParameters.NewColor *= FLinearColor::MakeFromColorTemperature(Light->Temperature);
-		}
-
 		FScene* Scene = this;
 		FLightSceneInfo* LightSceneInfo = Light->SceneProxy->GetLightSceneInfo();
 		ENQUEUE_RENDER_COMMAND(UpdateLightColorAndBrightness)(
