ColorManagement: Adding a color luminance getter to the color space class.

Additional comments/parameter names cleanup.

diff --git a/Engine/Source/Runtime/ColorManagement/Private/ColorSpace.cpp b/Engine/Source/Runtime/ColorManagement/Private/ColorSpace.cpp
--- a/Engine/Source/Runtime/ColorManagement/Private/ColorSpace.cpp
+++ b/Engine/Source/Runtime/ColorManagement/Private/ColorSpace.cpp
@@ -163,12 +163,12 @@ FMatrix44d FColorSpace::CalcRgbToXYZ() const
 	return Mat;
 }
 
-bool FColorSpace::Equals(const FColorSpace& CS, double Tolerance) const
+bool FColorSpace::Equals(const FColorSpace& ColorSpace, double Tolerance) const
 {
-	return	Chromaticities[0].Equals(CS.Chromaticities[0], Tolerance) &&
-		Chromaticities[1].Equals(CS.Chromaticities[1], Tolerance) &&
-		Chromaticities[2].Equals(CS.Chromaticities[2], Tolerance) &&
-		Chromaticities[3].Equals(CS.Chromaticities[3], Tolerance);
+	return	Chromaticities[0].Equals(ColorSpace.Chromaticities[0], Tolerance) &&
+		Chromaticities[1].Equals(ColorSpace.Chromaticities[1], Tolerance) &&
+		Chromaticities[2].Equals(ColorSpace.Chromaticities[2], Tolerance) &&
+		Chromaticities[3].Equals(ColorSpace.Chromaticities[3], Tolerance);
 }
 
 bool FColorSpace::IsSRGB() const
@@ -194,6 +194,12 @@ FLinearColor FColorSpace::MakeFromColorTemperature(float Temp) const
 	return FLinearColor((float)RGB.X, (float)RGB.Y, (float)RGB.Z);
 }
 
+float FColorSpace::GetLuminance(const FLinearColor& Color) const
+{
+	//Note: Equivalent to the dot product of Color and RgbToXYZ.GetColumn(1).
+	return Color.R * RgbToXYZ.M[0][1] + Color.G * RgbToXYZ.M[1][1] + Color.B * RgbToXYZ.M[2][1];
+}
+
 
 FMatrix44d FColorSpaceTransform::CalcChromaticAdaptionMatrix(FVector3d SourceXYZ, FVector3d TargetXYZ, EChromaticAdaptationMethod Method)
 {
diff --git a/Engine/Source/Runtime/ColorManagement/Public/ColorSpace.h b/Engine/Source/Runtime/ColorManagement/Public/ColorSpace.h
--- a/Engine/Source/Runtime/ColorManagement/Public/ColorSpace.h
+++ b/Engine/Source/Runtime/ColorManagement/Public/ColorSpace.h
@@ -9,7 +9,6 @@
 #include "Math/Matrix.h"
 #include "Math/UnrealMathSSE.h"
 #include "Math/Vector.h"
-#include "Math/Vector2D.h"
 #include "Misc/AssertionMacros.h"
 #include "Serialization/Archive.h"
 #include "Templates/UnrealTypeTraits.h"
@@ -223,35 +222,35 @@ public:
 	}
 
 	/**
-	 * Check against another vector for equality.
+	 * Check against another color space for equality.
 	 *
-	 * @param V The vector to check against.
+	 * @param ColorSpace The vector to check against.
 	 * @return true if the vectors are equal, false otherwise.
 	 */
-	FORCEINLINE bool operator==(const FColorSpace& CS) const
+	FORCEINLINE bool operator==(const FColorSpace& ColorSpace) const
 	{
-		return Chromaticities == CS.Chromaticities;
+		return Chromaticities == ColorSpace.Chromaticities;
 	}
 
 	/**
-	 * Check against another colorspace for inequality.
+	 * Check against another color space for inequality.
 	 *
-	 * @param V The vector to check against.
+	 * @param ColorSpace The vector to check against.
 	 * @return true if the vectors are not equal, false otherwise.
 	 */
-	FORCEINLINE bool operator!=(const FColorSpace& CS) const
+	FORCEINLINE bool operator!=(const FColorSpace& ColorSpace) const
 	{
-		return Chromaticities != CS.Chromaticities;
+		return Chromaticities != ColorSpace.Chromaticities;
 	}
 
 	/**
-	 * Check against another colorspace for equality, within specified error limits.
+	 * Check against another color space for equality, within specified error limits.
 	 *
-	 * @param V The vector to check against.
+	 * @param ColorSpace The vector to check against.
 	 * @param Tolerance Error tolerance.
 	 * @return true if the vectors are equal within tolerance limits, false otherwise.
 	 */
-	bool Equals(const FColorSpace& CS, double Tolerance = 1.e-7) const;
+	bool Equals(const FColorSpace& ColorSpace, double Tolerance = 1.e-7) const;
 
 	/**
 	 * Convenience function to verify if the color space matches the engine's default sRGB chromaticities.
@@ -265,6 +264,14 @@ public:
 	*/
 	FLinearColor MakeFromColorTemperature(float Temp) const;
 
+	/**
+	 * Calculate the color's luminance value in the current space.
+	 *
+	 * @param Color The sampled color.
+	 * @return float Luminance value.
+	 */
+	float GetLuminance(const FLinearColor& Color) const;
+
 private:
 
 	FMatrix44d CalcRgbToXYZ() const;
