ColorManagement: Cleanup of XYZ/xyY conversion functions.

ToYxy functions were serving no purpose so we deprecate them, and expose a new XYZ to xyY function. We also remove unnecessary FORCEINLINE.

diff --git a/Engine/Plugins/VirtualProduction/DMX/DMXPixelMapping/Source/DMXPixelMappingRuntime/Private/ColorSpace/DMXPixelMappingColorSpace_xyY.cpp b/Engine/Plugins/VirtualProduction/DMX/DMXPixelMapping/Source/DMXPixelMappingRuntime/Private/ColorSpace/DMXPixelMappingColorSpace_xyY.cpp
--- a/Engine/Plugins/VirtualProduction/DMX/DMXPixelMapping/Source/DMXPixelMappingRuntime/Private/ColorSpace/DMXPixelMappingColorSpace_xyY.cpp
+++ b/Engine/Plugins/VirtualProduction/DMX/DMXPixelMapping/Source/DMXPixelMappingRuntime/Private/ColorSpace/DMXPixelMappingColorSpace_xyY.cpp
@@ -17,28 +17,23 @@ void UDMXPixelMappingColorSpace_xyY::SetRGBA(const FLinearColor& InColor)
 		SRGBColorSpace;
 
 	// Convert RGB to CIE XYZ
-	const FMatrix44d Matrix = InputColorSpace.GetRgbToXYZ();
+	const FMatrix44d& Matrix = InputColorSpace.GetRgbToXYZ();
 	const FVector4 XYZW = Matrix.TransformVector(FVector(InColor));
-
-	const double SumXYZ = XYZW.X + XYZW.Y + XYZW.Z;
-
-	// Calculate chromaticity coordinates x and y, use white point if SumXYZ is zero
-	const double x = SumXYZ > 0.f ? XYZW.X / SumXYZ : 0.3127;
-	const double y = SumXYZ > 0.f ? XYZW.Y / SumXYZ : 0.3290;
+	const FVector3d xyY = UE::Color::XYZToxyY(XYZW);
 
 	if (XAttribute.IsValid())
 	{
-		SetAttributeValue(XAttribute, x);
+		SetAttributeValue(XAttribute, xyY[0]);
 	}
 
 	if (YAttribute.IsValid())
 	{
-		SetAttributeValue(YAttribute, y);
+		SetAttributeValue(YAttribute, xyY[1]);
 	}
 
 	if (LuminanceAttribute.IsValid())
 	{
-		const float Luminance = FMath::Clamp(XYZW.Y, MinLuminance, MaxLuminance);
+		const float Luminance = FMath::Clamp(xyY[2], MinLuminance, MaxLuminance);
 		SetAttributeValue(LuminanceAttribute, Luminance);
 	}
 }
diff --git a/Engine/Source/Runtime/ColorManagement/Public/ColorSpace.h b/Engine/Source/Runtime/ColorManagement/Public/ColorSpace.h
--- a/Engine/Source/Runtime/ColorManagement/Public/ColorSpace.h
+++ b/Engine/Source/Runtime/ColorManagement/Public/ColorSpace.h
@@ -20,7 +20,8 @@ namespace UE { namespace Color {
  *
  * @return FVector3d
  */
-FORCEINLINE FVector3d ToYxy(double LuminanceY, const FVector2d& Coordinate)
+UE_DEPRECATED(5.3, "ToYxy has been deprecated.")
+inline FVector3d ToYxy(double LuminanceY, const FVector2d& Coordinate)
 {
 	return FVector3d(LuminanceY, Coordinate.X, Coordinate.Y);
 }
@@ -30,34 +31,58 @@ FORCEINLINE FVector3d ToYxy(double LuminanceY, const FVector2d& Coordinate)
  *
  * @return FVector3d
  */
-FORCEINLINE FVector3d ToYxy(const FVector2d& Coordinate)
+UE_DEPRECATED(5.3, "ToYxy has been deprecated.")
+inline FVector3d ToYxy(const FVector2d& Coordinate)
 {
-	return ToYxy(1.0, Coordinate);
+	return FVector3d(1.0, Coordinate.X, Coordinate.Y);
 }
 
 /**
- * Convert coordinate to CIE XYZ tristimulus values with a luminance value.
- * 
+ * Convert chromaticity coordinate and luminance to CIE XYZ tristimulus values.
+ *
  * @return FVector3d
  */
-FORCEINLINE FVector3d ToXYZ(double LuminanceY, const FVector2d& Coordinate)
+inline FVector3d xyYToXYZ(const FVector3d& xyY)
 {
-	const FVector3d Yxy = ToYxy(LuminanceY, Coordinate);
+	double Divisor = FMath::Max(xyY[1], 1e-10);
+
 	return FVector3d(
-		Yxy[1] * Yxy[0] / FMath::Max(Yxy[2], 1e-10),
-		Yxy[0],
-		(1.0 - Yxy[1] - Yxy[2]) * Yxy[0] / FMath::Max(Yxy[2], 1e-10)
+		xyY[0] * xyY[2] / Divisor,
+		xyY[2],
+		(1.0 - xyY[0] - xyY[1]) * xyY[2] / Divisor
 	);
 }
 
 /**
- * Convert coordinate to CIE XYZ tristimulus values with a default luminance of 1.0.
+ * Convert CIE XYZ tristimulus values to chromaticitiy coordinates and luminance.
  *
  * @return FVector3d
  */
-FORCEINLINE FVector3d ToXYZ(const FVector2d& Coordinate)
+inline FVector3d XYZToxyY(const FVector3d& XYZ)
+{
+	double Divisor = XYZ[0] + XYZ[1] + XYZ[2];
+	if (Divisor == 0.0)
+	{
+		Divisor = 1e-10;
+	}
+
+	return FVector3d(
+		XYZ[0] / Divisor,
+		XYZ[1] / Divisor,
+		XYZ[1]
+	);
+}
+
+UE_DEPRECATED(5.3, "ToXYZ has been replaced by xyYToXYZ.")
+inline FVector3d ToXYZ(double LuminanceY, const FVector2d& Coordinate)
+{
+	return xyYToXYZ(FVector3d(Coordinate.X, Coordinate.Y, LuminanceY));
+}
+
+UE_DEPRECATED(5.3, "ToXYZ has been replaced by xyYToXYZ.")
+inline FVector3d ToXYZ(const FVector2d& Coordinate)
 {
-	return ToXYZ(1.0, Coordinate);
+	return xyYToXYZ(FVector3d(Coordinate.X, Coordinate.Y, 1.0));
 }
 
 /**
@@ -66,7 +91,7 @@ FORCEINLINE FVector3d ToXYZ(const FVector2d& Coordinate)
  * @return TMatrix<T>
  */
 template<typename T>
-FORCEINLINE UE::Math::TMatrix<T> Transpose(const FMatrix44d& Transform)
+inline UE::Math::TMatrix<T> Transpose(const FMatrix44d& Transform)
 {
 	if constexpr (std::is_same_v<T, double>)
 	{
@@ -82,7 +107,7 @@ FORCEINLINE UE::Math::TMatrix<T> Transpose(const FMatrix44d& Transform)
  * @param InWhitePoint White point type.
  * @return FVector2d Chromaticity coordinates.
  */
-FORCEINLINE FVector2d GetWhitePoint(EWhitePoint InWhitePoint)
+inline FVector2d GetWhitePoint(EWhitePoint InWhitePoint)
 {
 	switch (InWhitePoint)
 	{
@@ -153,7 +178,7 @@ public:
 	* @param OutWhite FVector2d for the white color chromaticity coordinate.
 	*/
 	template<typename T>
-	FORCEINLINE void GetChromaticities(UE::Math::TVector2<T>& OutRed, UE::Math::TVector2<T>& OutGreen, UE::Math::TVector2<T>& OutBlue, UE::Math::TVector2<T>& OutWhite) const
+	inline void GetChromaticities(UE::Math::TVector2<T>& OutRed, UE::Math::TVector2<T>& OutGreen, UE::Math::TVector2<T>& OutBlue, UE::Math::TVector2<T>& OutWhite) const
 	{
 		OutRed		= Chromaticities[0];
 		OutGreen	= Chromaticities[1];
@@ -166,7 +191,7 @@ public:
 	*
 	* @return FVector2d xy coordinates.
 	*/
-	FORCEINLINE const FVector2d& GetRedChromaticity() const
+	inline const FVector2d& GetRedChromaticity() const
 	{
 		return Chromaticities[0];
 	}
@@ -176,7 +201,7 @@ public:
 	*
 	* @return FVector2d xy coordinates.
 	*/
-	FORCEINLINE const FVector2d& GetGreenChromaticity() const
+	inline const FVector2d& GetGreenChromaticity() const
 	{
 		return Chromaticities[1];
 	}
@@ -186,7 +211,7 @@ public:
 	*
 	* @return FVector2d xy coordinates.
 	*/
-	FORCEINLINE const FVector2d& GetBlueChromaticity() const
+	inline const FVector2d& GetBlueChromaticity() const
 	{
 		return Chromaticities[2];
 	}
@@ -196,7 +221,7 @@ public:
 	*
 	* @return FVector2d xy coordinates.
 	*/
-	FORCEINLINE const FVector2d& GetWhiteChromaticity() const
+	inline const FVector2d& GetWhiteChromaticity() const
 	{
 		return Chromaticities[3];
 	}
@@ -206,7 +231,7 @@ public:
 	*
 	* @return FMatrix conversion matrix.
 	*/
-	FORCEINLINE const FMatrix44d& GetRgbToXYZ() const
+	inline const FMatrix44d& GetRgbToXYZ() const
 	{
 		return RgbToXYZ;
 	}
@@ -216,7 +241,7 @@ public:
 	*
 	* @return FMatrix conversion matrix.
 	*/
-	FORCEINLINE const FMatrix44d& GetXYZToRgb() const
+	inline const FMatrix44d& GetXYZToRgb() const
 	{
 		return XYZToRgb;
 	}
@@ -227,7 +252,7 @@ public:
 	 * @param ColorSpace The vector to check against.
 	 * @return true if the vectors are equal, false otherwise.
 	 */
-	FORCEINLINE bool operator==(const FColorSpace& ColorSpace) const
+	inline bool operator==(const FColorSpace& ColorSpace) const
 	{
 		return Chromaticities == ColorSpace.Chromaticities;
 	}
@@ -238,7 +263,7 @@ public:
 	 * @param ColorSpace The vector to check against.
 	 * @return true if the vectors are not equal, false otherwise.
 	 */
-	FORCEINLINE bool operator!=(const FColorSpace& ColorSpace) const
+	inline bool operator!=(const FColorSpace& ColorSpace) const
 	{
 		return Chromaticities != ColorSpace.Chromaticities;
 	}
