ColorManagement: Ensure the working color space has been loaded before use by shader compilation.

#rnx

diff --git a/Engine/Source/Runtime/ColorManagement/Private/ColorSpace.cpp b/Engine/Source/Runtime/ColorManagement/Private/ColorSpace.cpp
--- a/Engine/Source/Runtime/ColorManagement/Private/ColorSpace.cpp
+++ b/Engine/Source/Runtime/ColorManagement/Private/ColorSpace.cpp
@@ -1,20 +1,56 @@
 // Copyright Epic Games, Inc. All Rights Reserved.
 
 #include "ColorSpace.h"
+#include "Misc/ConfigCacheIni.h"
 #include "Math/VectorRegister.h"
 
 namespace UE { namespace Color {
 
+static bool bIsWorkingColorSpaceReadyForUse = false;
 static FColorSpace WorkingColorSpace = FColorSpace(EColorSpace::sRGB);
 
+void PreloadWorkingColorSpace()
+{
+	if (bIsWorkingColorSpaceReadyForUse)
+	{
+		return;
+	}
+
+	check(GConfig != nullptr && GConfig->IsReadyForUse());
+
+	bool bIsWorkingColorSpaceInConfig = true;
+	TStaticArray<FVector2d, 4> Chromaticities;
+	bIsWorkingColorSpaceInConfig &= GConfig->GetVector2D(TEXT("/Script/Engine.RendererSettings"), TEXT("RedChromaticityCoordinate"),   Chromaticities[0], GEngineIni);
+	bIsWorkingColorSpaceInConfig &= GConfig->GetVector2D(TEXT("/Script/Engine.RendererSettings"), TEXT("GreenChromaticityCoordinate"), Chromaticities[1], GEngineIni);
+	bIsWorkingColorSpaceInConfig &= GConfig->GetVector2D(TEXT("/Script/Engine.RendererSettings"), TEXT("BlueChromaticityCoordinate"),  Chromaticities[2], GEngineIni);
+	bIsWorkingColorSpaceInConfig &= GConfig->GetVector2D(TEXT("/Script/Engine.RendererSettings"), TEXT("WhiteChromaticityCoordinate"), Chromaticities[3], GEngineIni);
+
+	if (bIsWorkingColorSpaceInConfig)
+	{
+		FColorSpace::SetWorking(FColorSpace(Chromaticities[0], Chromaticities[1], Chromaticities[2], Chromaticities[3]));
+	}
+	else
+	{
+		// The working color space wasn't modified/serialized to the config so we keep the implicit sRGB default.
+		bIsWorkingColorSpaceReadyForUse = true;
+	}
+};
+
 const FColorSpace& FColorSpace::GetWorking()
 {
+	/**
+	 * NOTE: Addresses issue where shader compilation can request the working color space before it has been loaded by renderer settings.
+	 * We optionally early-load the settings here and leave existing singleton interface as is.
+	 */
+	UE_CALL_ONCE(PreloadWorkingColorSpace);
+
 	return WorkingColorSpace;
 }
 
 void FColorSpace::SetWorking(FColorSpace ColorSpace)
 {
 	WorkingColorSpace = MoveTemp(ColorSpace);
+	bIsWorkingColorSpaceReadyForUse = true;
 }
 
 static bool IsSRGBChromaticities(const TStaticArray<FVector2d, 4>& Chromaticities, double Tolerance = 1.e-7)
