ColorManagement: Unifying use of FGenericPlatformMath in transfer functions.

#rnx

diff --git a/Engine/Source/Runtime/ColorManagement/Public/TransferFunctions.h b/Engine/Source/Runtime/ColorManagement/Public/TransferFunctions.h
--- a/Engine/Source/Runtime/ColorManagement/Public/TransferFunctions.h
+++ b/Engine/Source/Runtime/ColorManagement/Public/TransferFunctions.h
@@ -83,7 +83,7 @@ FORCEINLINE float DecodeST2084(float Value)
 	const float c3 = 2392 / 4096.0f * 32.f;
 
 	const float Vp = FGenericPlatformMath::Pow(Value, 1.0f / m2);
-	Value = FMath::Max(0.0f, Vp - c1);
+	Value = FGenericPlatformMath::Max(0.0f, Vp - c1);
 	return FGenericPlatformMath::Pow((Value / (c2 - c3 * Vp)), 1.0f / m1) * Lp;
 }
 
