ColorManagement: Added missing enum max values.

diff --git a/Engine/Source/Runtime/ColorManagement/Public/ColorManagementDefines.h b/Engine/Source/Runtime/ColorManagement/Public/ColorManagementDefines.h
--- a/Engine/Source/Runtime/ColorManagement/Public/ColorManagementDefines.h
+++ b/Engine/Source/Runtime/ColorManagement/Public/ColorManagementDefines.h
@@ -59,6 +59,7 @@ enum class EColorSpace : uint8
 	GoProProtuneNative = 12,
 	PanasonicVGamut = 13,
 	PLASA_E1_54 = 14,
+	Max,
 };
 
 
@@ -71,6 +72,7 @@ enum class EChromaticAdaptationMethod : uint8
 	None = 0,
 	Bradford = 1,
 	CAT02 = 2,
+	Max,
 };
 
 /** List of standard white points. */
@@ -79,6 +81,7 @@ enum class EWhitePoint : uint8
 	CIE1931_D65 = 0,
 	ACES_D60 = 1,
 	DCI_CalibrationWhite,
+	Max,
 };
 
 } } // end namespace UE::Color
diff --git a/Engine/Source/Runtime/Engine/Classes/Engine/TextureDefines.h b/Engine/Source/Runtime/Engine/Classes/Engine/TextureDefines.h
--- a/Engine/Source/Runtime/Engine/Classes/Engine/TextureDefines.h
+++ b/Engine/Source/Runtime/Engine/Classes/Engine/TextureDefines.h
@@ -429,6 +429,7 @@ enum class ETextureChromaticAdaptationMethod : uint8
 	TCAM_None		= 0 UMETA(DisplayName = "None", ToolTip = "No chromatic adaptation is applied."),
 	TCAM_Bradford	= 1 UMETA(DisplayName = "Bradford", ToolTip = "Chromatic adaptation is applied using the Bradford method."),
 	TCAM_CAT02		= 2 UMETA(DisplayName = "CAT02", ToolTip = "Chromatic adaptation is applied using the CAT02 method."),
+	TCAM_MAX,
 };
 
 
