ColorManagement: Adding luminance factor getter.

#rnx

diff --git a/Engine/Source/Runtime/ColorManagement/Private/ColorSpace.cpp b/Engine/Source/Runtime/ColorManagement/Private/ColorSpace.cpp
--- a/Engine/Source/Runtime/ColorManagement/Private/ColorSpace.cpp
+++ b/Engine/Source/Runtime/ColorManagement/Private/ColorSpace.cpp
@@ -236,6 +236,11 @@ float FColorSpace::GetLuminance(const FLinearColor& Color) const
 	return Color.R * RgbToXYZ.M[0][1] + Color.G * RgbToXYZ.M[1][1] + Color.B * RgbToXYZ.M[2][1];
 }
 
+FLinearColor FColorSpace::GetLuminanceFactors() const
+{
+	return FLinearColor(RgbToXYZ.GetColumn(1));
+}
+
 
 FMatrix44d FColorSpaceTransform::CalcChromaticAdaptionMatrix(FVector3d SourceXYZ, FVector3d TargetXYZ, EChromaticAdaptationMethod Method)
 {
diff --git a/Engine/Source/Runtime/ColorManagement/Private/Tests/ColorManagementTests.cpp b/Engine/Source/Runtime/ColorManagement/Private/Tests/ColorManagementTests.cpp
--- a/Engine/Source/Runtime/ColorManagement/Private/Tests/ColorManagementTests.cpp
+++ b/Engine/Source/Runtime/ColorManagement/Private/Tests/ColorManagementTests.cpp
@@ -148,7 +148,7 @@ static void TestColorSpaceTransforms(UE::Color::EChromaticAdaptationMethod Metho
 
 	using namespace UE::Color;
 
-	double Tolerance = SMALL_NUMBER;
+	double Tolerance = UE_SMALL_NUMBER;
 
 	const FColorSpace Src = FColorSpace(EColorSpace::ACESAP1);
 	FMatrix44d Mat0, Mat1;
@@ -258,6 +258,22 @@ static void TestColorSpaceTransforms(UE::Color::EChromaticAdaptationMethod Metho
 
 }
 
+static void TestLuminance()
+{
+	using namespace UE::Color;
+
+	// Note: test factors from the python (colour-science) colour library.
+
+	FLinearColor LuminanceFactors = FColorSpace(EColorSpace::sRGB).GetLuminanceFactors();
+	LogTest(TEXT("sRGB luminance factors equality test"), LuminanceFactors.Equals(FLinearColor(0.212639005872, 0.715168678768, 0.0721923153607), UE_SMALL_NUMBER));
+
+	LuminanceFactors = FColorSpace(EColorSpace::ACESAP1).GetLuminanceFactors();
+	LogTest(TEXT("ACESAP1 luminance factors equality test"), LuminanceFactors.Equals(FLinearColor(0.272228716781, 0.674081765811, 0.0536895174079), UE_SMALL_NUMBER));
+
+	LuminanceFactors = FColorSpace(EColorSpace::Rec2020).GetLuminanceFactors();
+	LogTest(TEXT("Rec2020 luminance factors equality test"), LuminanceFactors.Equals(FLinearColor(0.262700212011, 0.677998071519, 0.0593017164699), UE_SMALL_NUMBER));
+}
+
 IMPLEMENT_SIMPLE_AUTOMATION_TEST(FColorSpaceTest, "System.ColorManagement.ColorSpace", EAutomationTestFlags::EditorContext | EAutomationTestFlags::EngineFilter)
 bool FColorSpaceTest::RunTest(const FString& Parameters)
 {
@@ -276,6 +292,8 @@ bool FColorSpaceTest::RunTest(const FString& Parameters)
 
 	TestColorSpaceTransforms(EChromaticAdaptationMethod::Bradford);
 
+	TestLuminance();
+
 	return GPassing;
 }
 
diff --git a/Engine/Source/Runtime/ColorManagement/Public/ColorSpace.h b/Engine/Source/Runtime/ColorManagement/Public/ColorSpace.h
--- a/Engine/Source/Runtime/ColorManagement/Public/ColorSpace.h
+++ b/Engine/Source/Runtime/ColorManagement/Public/ColorSpace.h
@@ -297,6 +297,13 @@ public:
 	 */
 	float GetLuminance(const FLinearColor& Color) const;
 
+	/**
+	 * Get the luminance factors in the current space.
+	 *
+	 * @return FLinearColor Luminance factors.
+	 */
+	FLinearColor GetLuminanceFactors() const;
+
 private:
 
 	FMatrix44d CalcRgbToXYZ() const;
