ColorManagement: Fix crash on Linux with unaligned color space transform.

#rnx

diff --git a/Engine/Source/Runtime/ColorManagement/Private/ColorSpace.cpp b/Engine/Source/Runtime/ColorManagement/Private/ColorSpace.cpp
--- a/Engine/Source/Runtime/ColorManagement/Private/ColorSpace.cpp
+++ b/Engine/Source/Runtime/ColorManagement/Private/ColorSpace.cpp
@@ -341,9 +341,9 @@ FColorSpaceTransform::FColorSpaceTransform(FMatrix44d Matrix)
 FLinearColor FColorSpaceTransform::Apply(const FLinearColor& Color) const
 {
 	FLinearColor Result;
-	VectorRegister4Float VecP = VectorLoadAligned((const FVector4f*)&Color);
+	VectorRegister4Float VecP = VectorLoad(&Color.R);
 	VectorRegister4Float VecR = VectorTransformVector(VecP, this);
-	VectorStoreAligned(VecR, (FVector4f*)&Result);
+	VectorStore(VecR, &Result.R);
 	return Result;
 }
 
diff --git a/Engine/Source/Runtime/ColorManagement/Private/Tests/ColorManagementTests.cpp b/Engine/Source/Runtime/ColorManagement/Private/Tests/ColorManagementTests.cpp
--- a/Engine/Source/Runtime/ColorManagement/Private/Tests/ColorManagementTests.cpp
+++ b/Engine/Source/Runtime/ColorManagement/Private/Tests/ColorManagementTests.cpp
@@ -255,7 +255,29 @@ static void TestColorSpaceTransforms(UE::Color::EChromaticAdaptationMethod Metho
 		).GetTransposed();
 		LogTest(TEXT("AP1->Rec709 with Bradford chromatic adaptation"), TestMatricesDoubleEqual(Mat0, Mat1, Tolerance));
 	}
+}
 
+static void TestColorTransforms()
+{
+	using namespace UE::Color;
+
+	// Intentionally misalign the start of FLinearColor member,
+	// as a test against 16-byte (128bit) alignment when using 4 float SIMD instructions.
+	// 
+	// Given that we now load/store from unaligned memory, the above alignment requirement is lifted.
+	struct FAlignmentTest
+	{
+		float Nudge = 0.0f;
+		FLinearColor SrcColor = FLinearColor(1.0f, 0.5f, 0.0f);
+	};
+
+	FColorSpaceTransform Transform = FColorSpaceTransform(FColorSpace(EColorSpace::sRGB), FColorSpace(EColorSpace::ACESAP1), EChromaticAdaptationMethod::Bradford);
+
+	TUniquePtr<FAlignmentTest> AlignmentTest = MakeUnique<FAlignmentTest>();
+	FLinearColor ExpectedResult = FLinearColor(0.78285898f, 0.52837066f, 0.07540048f);
+	FLinearColor Result = Transform.Apply(AlignmentTest->SrcColor);
+
+	LogTest(TEXT("FLinearColor sRGB->AP1 color transform"), Result.Equals(ExpectedResult, UE_SMALL_NUMBER));
 }
 
 static void TestLuminance()
@@ -292,6 +314,8 @@ bool FColorSpaceTest::RunTest(const FString& Parameters)
 
 	TestColorSpaceTransforms(EChromaticAdaptationMethod::Bradford);
 
+	TestColorTransforms();
+
 	TestLuminance();
 
 	return GPassing;
