ColorManagement: Convert tests to low-level tests.

#rnx

diff --git a/Engine/Source/Programs/LowLevelTests/FoundationTests.Build.cs b/Engine/Source/Programs/LowLevelTests/FoundationTests.Build.cs
--- a/Engine/Source/Programs/LowLevelTests/FoundationTests.Build.cs
+++ b/Engine/Source/Programs/LowLevelTests/FoundationTests.Build.cs
@@ -30,7 +30,8 @@ public class FoundationTests : TestModuleRules
 				"Cbor",
 				"CoreUObject",
 				"TelemetryUtils",
-				"AssetRegistry"
+				"AssetRegistry",
+				"ColorManagement"
 			});
 
 		if (Target.bBuildWithEditorOnlyData)
diff --git a/Engine/Source/Runtime/ColorManagement/Private/Tests/ColorManagementTests.cpp b/Engine/Source/Runtime/ColorManagement/Private/Tests/ColorManagementTests.cpp
--- a/Engine/Source/Runtime/ColorManagement/Private/Tests/ColorManagementTests.cpp
+++ b/Engine/Source/Runtime/ColorManagement/Private/Tests/ColorManagementTests.cpp
@@ -1,5 +1,7 @@
 // Copyright Epic Games, Inc. All Rights Reserved.
 
+#if WITH_TESTS
+
 #include "CoreTypes.h"
 #include "Containers/UnrealString.h"
 #include "TransferFunctions.h"
@@ -7,23 +9,14 @@
 #include "Math/NumericLimits.h"
 #include "Math/UnrealMathUtility.h"
 #include "Logging/LogMacros.h"
-#include "Misc/AutomationTest.h"
-
-#if WITH_DEV_AUTOMATION_TESTS
-
-static bool GPassing;
-static double GdiffDouble = 0.0;
-static int32 GRow = 0;
-static int32 GColumn = 0;
+#include "Tests/TestHarnessAdapter.h"
 
 DEFINE_LOG_CATEGORY_STATIC(LogUnrealColorManagementTest, Log, All);
 
-IMPLEMENT_SIMPLE_AUTOMATION_TEST(FTransferFunctionsTest, "System.ColorManagement.TransferFunctions", EAutomationTestFlags::EditorContext | EAutomationTestFlags::EngineFilter)
-bool FTransferFunctionsTest::RunTest(const FString& Parameters)
+TEST_CASE_NAMED(FTransferFunctionsTest, "System::ColorManagement::TransferFunctions", "[EditorContext][EngineFilter]")
 {
 	using namespace UE::Color;
 
-	GPassing = true;
 	const float TestIncrement = 0.05f;
 
 	// Verify that all transfer functions correctly inverse each other.
@@ -36,10 +29,9 @@ bool FTransferFunctionsTest::RunTest(const FString& Parameters)
 		{
 			float Encoded = UE::Color::Encode(EncodingType, TestValue);
 			float Decoded = UE::Color::Decode(EncodingType, Encoded);
-			GPassing &= TestEqual(TEXT("Transfer function encode followed by decode must match identity"), Decoded, TestValue, KINDA_SMALL_NUMBER);
+			CHECK_MESSAGE(TEXT("Transfer function encode followed by decode must match identity"), FMath::IsNearlyEqual(Decoded, TestValue, KINDA_SMALL_NUMBER));
 		}
 	}
-	return GPassing;
 }
 
 /**
@@ -60,7 +52,7 @@ static bool TestMatricesDoubleEqual(FMatrix44d& Mat0, FMatrix44d& Mat1, double T
 			double Diff = Mat0.M[Row][Column] - Mat1.M[Row][Column];
 			if (FMath::Abs(Diff) > Tolerance)
 			{
-				GdiffDouble = Diff; GRow = Row; GColumn = Column;
+				UE_LOG(LogUnrealColorManagementTest, Log, TEXT("Bad(%.8f) at [%i, %i]"), Diff, Row, Column);
 				return false;
 			}
 		}
@@ -68,16 +60,6 @@ static bool TestMatricesDoubleEqual(FMatrix44d& Mat0, FMatrix44d& Mat1, double T
 	return true;
 }
 
-FORCENOINLINE void LogTest(const TCHAR* TestName, bool bHasPassed)
-{
-	if (bHasPassed == false)
-	{
-		UE_LOG(LogUnrealColorManagementTest, Log, TEXT("%s <double>: %s"), bHasPassed ? TEXT("PASSED") : TEXT("FAILED"), TestName);
-		UE_LOG(LogUnrealColorManagementTest, Log, TEXT("Bad(%.8f) at [%i, %i]"), GdiffDouble, GRow, GColumn);
-		GPassing = false;
-	}
-}
-
 static void TestRgbToXyzConversionMatrices()
 {
 	// Note: test matrices generated using the python (colour-science) colour library.
@@ -93,7 +75,7 @@ static void TestRgbToXyzConversionMatrices()
 		{0.0193308187156, 0.119194779795, 0.95053215225, 0.0},
 		{0,0,0,1}
 	).GetTransposed();
-	LogTest(TEXT("Rec709 RGB2XYZ Matrix Equality"), TestMatricesDoubleEqual(Mat0, Mat1, Tolerance));
+	CHECK_MESSAGE(TEXT("Rec709 RGB2XYZ Matrix Equality"), TestMatricesDoubleEqual(Mat0, Mat1, Tolerance));
 
 	Mat0 = FColorSpace(EColorSpace::Rec2020).GetRgbToXYZ();
 	Mat1 = FMatrix44d(
@@ -102,7 +84,7 @@ static void TestRgbToXyzConversionMatrices()
 		{4.99410657447e-17, 0.0280726930491, 1.06098505771, 0.0},
 		{0,0,0,1}
 	).GetTransposed();
-	LogTest(TEXT("Rec2020 RGB2XYZ Matrix Equality"), TestMatricesDoubleEqual(Mat0, Mat1, Tolerance));
+	CHECK_MESSAGE(TEXT("Rec2020 RGB2XYZ Matrix Equality"), TestMatricesDoubleEqual(Mat0, Mat1, Tolerance));
 
 	Mat0 = FColorSpace(EColorSpace::ACESAP0).GetRgbToXYZ();
 	Mat1 = FMatrix44d(
@@ -111,7 +93,7 @@ static void TestRgbToXyzConversionMatrices()
 		{ 0.0, 0.0, 1.0088251844, 0.0 },
 		{ 0,0,0,1 }
 	).GetTransposed();
-	LogTest(TEXT("AP0 RGB2XYZ Matrix Equality"), TestMatricesDoubleEqual(Mat0, Mat1, Tolerance));
+	CHECK_MESSAGE(TEXT("AP0 RGB2XYZ Matrix Equality"), TestMatricesDoubleEqual(Mat0, Mat1, Tolerance));
 
 	Mat0 = FColorSpace(EColorSpace::ACESAP1).GetRgbToXYZ();
 	Mat1 = FMatrix44d(
@@ -120,7 +102,7 @@ static void TestRgbToXyzConversionMatrices()
 		{ -0.00557464949039, 0.00406073352898, 1.01033910031, 0.0 },
 		{ 0,0,0,1 }
 	).GetTransposed();
-	LogTest(TEXT("AP1 RGB2XYZ Matrix Equality"), TestMatricesDoubleEqual(Mat0, Mat1, Tolerance));
+	CHECK_MESSAGE(TEXT("AP1 RGB2XYZ Matrix Equality"), TestMatricesDoubleEqual(Mat0, Mat1, Tolerance));
 
 	Mat0 = FColorSpace(EColorSpace::P3DCI).GetRgbToXYZ();
 	Mat1 = FMatrix44d(
@@ -129,7 +111,7 @@ static void TestRgbToXyzConversionMatrices()
 		{-3.63410131697e-17, 0.047060560054, 0.907355394362, 0.0},
 		{0,0,0,1}
 	).GetTransposed();
-	LogTest(TEXT("P3DCI RGB2XYZ Matrix Equality"), TestMatricesDoubleEqual(Mat0, Mat1, Tolerance));
+	CHECK_MESSAGE(TEXT("P3DCI RGB2XYZ Matrix Equality"), TestMatricesDoubleEqual(Mat0, Mat1, Tolerance));
 
 	Mat0 = FColorSpace(EColorSpace::P3D65).GetRgbToXYZ();
 	Mat1 = FMatrix44d(
@@ -138,7 +120,7 @@ static void TestRgbToXyzConversionMatrices()
 		{-3.97207551693e-17, 0.0451133818589, 1.0439443689, 0.0},
 		{0,0,0,1}
 	).GetTransposed();
-	LogTest(TEXT("P3D65 RGB2XYZ Matrix Equality"), TestMatricesDoubleEqual(Mat0, Mat1, Tolerance));
+	CHECK_MESSAGE(TEXT("P3D65 RGB2XYZ Matrix Equality"), TestMatricesDoubleEqual(Mat0, Mat1, Tolerance));
 }
 
 
@@ -162,7 +144,7 @@ static void TestColorSpaceTransforms(UE::Color::EChromaticAdaptationMethod Metho
 			{ -0.0245682525938, -0.125750404281, 1.06563695775, 0.0 },
 			{0,0,0,1}
 		).GetTransposed();
-		LogTest(TEXT("AP1->Rec709 without chromatic adaptation"), TestMatricesDoubleEqual(Mat0, Mat1, Tolerance));
+		CHECK_MESSAGE(TEXT("AP1->Rec709 without chromatic adaptation"), TestMatricesDoubleEqual(Mat0, Mat1, Tolerance));
 
 		Mat0 = FColorSpaceTransform(Src, FColorSpace(EColorSpace::Rec2020), Method);
 		Mat1 = FMatrix44d(
@@ -171,7 +153,7 @@ static void TestColorSpaceTransforms(UE::Color::EChromaticAdaptationMethod Metho
 			{ -0.005209686529, -0.0226414456785, 0.95230241486, 0.0 },
 			{ 0,0,0,1 }
 		).GetTransposed();
-		LogTest(TEXT("AP1->Rec709 without chromatic adaptation"), TestMatricesDoubleEqual(Mat0, Mat1, Tolerance));
+		CHECK_MESSAGE(TEXT("AP1->Rec709 without chromatic adaptation"), TestMatricesDoubleEqual(Mat0, Mat1, Tolerance));
 
 		Mat0 = FColorSpaceTransform(Src, FColorSpace(EColorSpace::ACESAP0), Method);
 		Mat1 = FMatrix44d(
@@ -180,11 +162,11 @@ static void TestColorSpaceTransforms(UE::Color::EChromaticAdaptationMethod Metho
 			{ -0.00552588255811, 0.00402521030598, 1.00150067225, 0.0 },
 			{ 0,0,0,1 }
 		).GetTransposed();
-		LogTest(TEXT("AP1->Rec709 without chromatic adaptation"), TestMatricesDoubleEqual(Mat0, Mat1, Tolerance));
+		CHECK_MESSAGE(TEXT("AP1->Rec709 without chromatic adaptation"), TestMatricesDoubleEqual(Mat0, Mat1, Tolerance));
 
 		Mat0 = FColorSpaceTransform(Src, FColorSpace(EColorSpace::ACESAP1), Method);
 		Mat1 = FMatrix44d::Identity;
-		LogTest(TEXT("AP1->Rec709 without chromatic adaptation"), TestMatricesDoubleEqual(Mat0, Mat1, Tolerance));
+		CHECK_MESSAGE(TEXT("AP1->Rec709 without chromatic adaptation"), TestMatricesDoubleEqual(Mat0, Mat1, Tolerance));
 
 		Mat0 = FColorSpaceTransform(Src, FColorSpace(EColorSpace::P3DCI), Method);
 		Mat1 = FMatrix44d(
@@ -193,7 +175,7 @@ static void TestColorSpaceTransforms(UE::Color::EChromaticAdaptationMethod Metho
 			{ -0.0026742897488, -0.0490786970568, 1.11404817691, 0.0 },
 			{ 0,0,0,1 }
 		).GetTransposed();
-		LogTest(TEXT("AP1->Rec709 without chromatic adaptation"), TestMatricesDoubleEqual(Mat0, Mat1, Tolerance));
+		CHECK_MESSAGE(TEXT("AP1->Rec709 without chromatic adaptation"), TestMatricesDoubleEqual(Mat0, Mat1, Tolerance));
 
 		Mat0 = FColorSpaceTransform(Src, FColorSpace(EColorSpace::P3D65), Method);
 		Mat1 = FMatrix44d(
@@ -202,7 +184,7 @@ static void TestColorSpaceTransforms(UE::Color::EChromaticAdaptationMethod Metho
 			{ -0.0023243874884, -0.0426572735573, 0.968286867584, 0.0 },
 			{ 0,0,0,1 }
 		).GetTransposed();
-		LogTest(TEXT("AP1->Rec709 without chromatic adaptation"), TestMatricesDoubleEqual(Mat0, Mat1, Tolerance));
+		CHECK_MESSAGE(TEXT("AP1->Rec709 without chromatic adaptation"), TestMatricesDoubleEqual(Mat0, Mat1, Tolerance));
 	}
 	else if (Method == UE::Color::EChromaticAdaptationMethod::Bradford)
 	{
@@ -213,7 +195,7 @@ static void TestColorSpaceTransforms(UE::Color::EChromaticAdaptationMethod Metho
 			{ -0.0240033568046, -0.128968976065, 1.15297233287, 0.0 },
 			{ 0,0,0,1 }
 		).GetTransposed();
-		LogTest(TEXT("AP1->Rec709 with Bradford chromatic adaptation"), TestMatricesDoubleEqual(Mat0, Mat1, Tolerance));
+		CHECK_MESSAGE(TEXT("AP1->Rec709 with Bradford chromatic adaptation"), TestMatricesDoubleEqual(Mat0, Mat1, Tolerance));
 
 		Mat0 = FColorSpaceTransform(Src, FColorSpace(EColorSpace::Rec2020), Method);
 		Mat1 = FMatrix44d(
@@ -222,7 +204,7 @@ static void TestColorSpaceTransforms(UE::Color::EChromaticAdaptationMethod Metho
 			{ -0.00501335146809, -0.0252900718108, 1.03030342328, 0.0 },
 			{ 0,0,0,1 }
 		).GetTransposed();
-		LogTest(TEXT("AP1->Rec709 with Bradford chromatic adaptation"), TestMatricesDoubleEqual(Mat0, Mat1, Tolerance));
+		CHECK_MESSAGE(TEXT("AP1->Rec709 with Bradford chromatic adaptation"), TestMatricesDoubleEqual(Mat0, Mat1, Tolerance));
 
 		Mat0 = FColorSpaceTransform(Src, FColorSpace(EColorSpace::ACESAP0), Method);
 		Mat1 = FMatrix44d(
@@ -231,11 +213,11 @@ static void TestColorSpaceTransforms(UE::Color::EChromaticAdaptationMethod Metho
 			{ -0.00552588255811, 0.00402521030598, 1.00150067225, 0.0 },
 			{ 0,0,0,1 }
 		).GetTransposed();
-		LogTest(TEXT("AP1->Rec709 with Bradford chromatic adaptation"), TestMatricesDoubleEqual(Mat0, Mat1, Tolerance));
+		CHECK_MESSAGE(TEXT("AP1->Rec709 with Bradford chromatic adaptation"), TestMatricesDoubleEqual(Mat0, Mat1, Tolerance));
 
 		Mat0 = FColorSpaceTransform(Src, FColorSpace(EColorSpace::ACESAP1), Method);
 		Mat1 = FMatrix44d::Identity;
-		LogTest(TEXT("AP1->Rec709 with Bradford chromatic adaptation"), TestMatricesDoubleEqual(Mat0, Mat1, Tolerance));
+		CHECK_MESSAGE(TEXT("AP1->Rec709 with Bradford chromatic adaptation"), TestMatricesDoubleEqual(Mat0, Mat1, Tolerance));
 
 		Mat0 = FColorSpaceTransform(Src, FColorSpace(EColorSpace::P3DCI), Method);
 		Mat1 = FMatrix44d(
@@ -244,7 +226,7 @@ static void TestColorSpaceTransforms(UE::Color::EChromaticAdaptationMethod Metho
 			{ -0.00255286167095, -0.0470296027287, 1.0495824644, 0.0 },
 			{ 0,0,0,1 }
 		).GetTransposed();
-		LogTest(TEXT("AP1->Rec709 with Bradford chromatic adaptation"), TestMatricesDoubleEqual(Mat0, Mat1, Tolerance));
+		CHECK_MESSAGE(TEXT("AP1->Rec709 with Bradford chromatic adaptation"), TestMatricesDoubleEqual(Mat0, Mat1, Tolerance));
 
 		Mat0 = FColorSpaceTransform(Src, FColorSpace(EColorSpace::P3D65), Method);
 		Mat1 = FMatrix44d(
@@ -253,11 +235,11 @@ static void TestColorSpaceTransforms(UE::Color::EChromaticAdaptationMethod Metho
 			{ -0.00215900951357, -0.0454593248373, 1.04761833435, 0.0 },
 			{ 0,0,0,1 }
 		).GetTransposed();
-		LogTest(TEXT("AP1->Rec709 with Bradford chromatic adaptation"), TestMatricesDoubleEqual(Mat0, Mat1, Tolerance));
+		CHECK_MESSAGE(TEXT("AP1->Rec709 with Bradford chromatic adaptation"), TestMatricesDoubleEqual(Mat0, Mat1, Tolerance));
 	}
 }
 
-static void TestColorTransforms()
+static void TestAppliedColorTransforms()
 {
 	using namespace UE::Color;
 
@@ -277,7 +259,7 @@ static void TestColorTransforms()
 	FLinearColor ExpectedResult = FLinearColor(0.78285898f, 0.52837066f, 0.07540048f);
 	FLinearColor Result = Transform.Apply(AlignmentTest->SrcColor);
 
-	LogTest(TEXT("FLinearColor sRGB->AP1 color transform"), Result.Equals(ExpectedResult, UE_SMALL_NUMBER));
+	CHECK_MESSAGE(TEXT("FLinearColor sRGB->AP1 color transform"), Result.Equals(ExpectedResult, UE_SMALL_NUMBER));
 }
 
 static void TestLuminance()
@@ -287,39 +269,38 @@ static void TestLuminance()
 	// Note: test factors from the python (colour-science) colour library.
 
 	FLinearColor LuminanceFactors = FColorSpace(EColorSpace::sRGB).GetLuminanceFactors();
-	LogTest(TEXT("sRGB luminance factors equality test"), LuminanceFactors.Equals(FLinearColor(0.212639005872, 0.715168678768, 0.0721923153607), UE_SMALL_NUMBER));
+	CHECK_MESSAGE(TEXT("sRGB luminance factors equality test"), LuminanceFactors.Equals(FLinearColor(0.212639005872, 0.715168678768, 0.0721923153607), UE_SMALL_NUMBER));
 
 	LuminanceFactors = FColorSpace(EColorSpace::ACESAP1).GetLuminanceFactors();
-	LogTest(TEXT("ACESAP1 luminance factors equality test"), LuminanceFactors.Equals(FLinearColor(0.272228716781, 0.674081765811, 0.0536895174079), UE_SMALL_NUMBER));
+	CHECK_MESSAGE(TEXT("ACESAP1 luminance factors equality test"), LuminanceFactors.Equals(FLinearColor(0.272228716781, 0.674081765811, 0.0536895174079), UE_SMALL_NUMBER));
 
 	LuminanceFactors = FColorSpace(EColorSpace::Rec2020).GetLuminanceFactors();
-	LogTest(TEXT("Rec2020 luminance factors equality test"), LuminanceFactors.Equals(FLinearColor(0.262700212011, 0.677998071519, 0.0593017164699), UE_SMALL_NUMBER));
+	CHECK_MESSAGE(TEXT("Rec2020 luminance factors equality test"), LuminanceFactors.Equals(FLinearColor(0.262700212011, 0.677998071519, 0.0593017164699), UE_SMALL_NUMBER));
 }
 
-IMPLEMENT_SIMPLE_AUTOMATION_TEST(FColorSpaceTest, "System.ColorManagement.ColorSpace", EAutomationTestFlags::EditorContext | EAutomationTestFlags::EngineFilter)
-bool FColorSpaceTest::RunTest(const FString& Parameters)
+TEST_CASE_NAMED(FColorSpaceTest, "System::ColorManagement::ColorSpace", "[EditorContext][EngineFilter]")
 {
 	using namespace UE::Color;
 
-	GPassing = true;
-
 	FColorSpace CS = FColorSpace(EColorSpace::sRGB);
 	FMatrix44d Mat0 = FColorSpaceTransform(CS, CS, EChromaticAdaptationMethod::None);
 	FMatrix44d Mat1 = FMatrix44d::Identity;
-	LogTest(TEXT("Identity color space conversion to itself should match identity"), TestMatricesDoubleEqual(Mat0, Mat1, 0.00000001));
+	CHECK_MESSAGE(TEXT("Identity color space conversion to itself should match identity"), TestMatricesDoubleEqual(Mat0, Mat1, 0.00000001));
 
+	SECTION("RgbToXyzConversionMatrices")
 	TestRgbToXyzConversionMatrices();
 
+	SECTION("ColorSpaceTransforms")
 	TestColorSpaceTransforms(EChromaticAdaptationMethod::None);
 
+	SECTION("ColorSpaceTransforms_Bradford")
 	TestColorSpaceTransforms(EChromaticAdaptationMethod::Bradford);
 
-	TestColorTransforms();
+	SECTION("AppliedColorTransforms")
+	TestAppliedColorTransforms();
 
+	SECTION("Luminance")
 	TestLuminance();
-
-	return GPassing;
 }
 
-
-#endif
\ No newline at end of file
+#endif //WITH_TESTS
