ColorManagement: Deprecated function cleanup and removal of forceinline.

#rnx

diff --git a/Engine/Source/Runtime/ColorManagement/Private/TransferFunctions.cpp b/Engine/Source/Runtime/ColorManagement/Private/TransferFunctions.cpp
--- a/Engine/Source/Runtime/ColorManagement/Private/TransferFunctions.cpp
+++ b/Engine/Source/Runtime/ColorManagement/Private/TransferFunctions.cpp
@@ -7,7 +7,7 @@ DEFINE_LOG_CATEGORY_STATIC(LogTransferFunctions, Log, All);
 
 namespace UE { namespace Color {
 
-FORCEINLINE TFunction<float(float)> GetTransferFunction(EEncoding SourceEncoding, bool bIsEncode)
+inline TFunction<float(float)> GetTransferFunction(EEncoding SourceEncoding, bool bIsEncode)
 {
 	switch (SourceEncoding)
 	{
diff --git a/Engine/Source/Runtime/ColorManagement/Public/ColorSpace.h b/Engine/Source/Runtime/ColorManagement/Public/ColorSpace.h
--- a/Engine/Source/Runtime/ColorManagement/Public/ColorSpace.h
+++ b/Engine/Source/Runtime/ColorManagement/Public/ColorSpace.h
@@ -15,28 +15,6 @@
 
 namespace UE { namespace Color {
 
-/**
- * Convert coordinate to CIE Yxy with a luminance value.
- *
- * @return FVector3d
- */
-UE_DEPRECATED(5.3, "ToYxy has been deprecated.")
-inline FVector3d ToYxy(double LuminanceY, const FVector2d& Coordinate)
-{
-	return FVector3d(LuminanceY, Coordinate.X, Coordinate.Y);
-}
-
-/**
- * Convert coordinate to CIE Yxy with a default luminance of 1.0.
- *
- * @return FVector3d
- */
-UE_DEPRECATED(5.3, "ToYxy has been deprecated.")
-inline FVector3d ToYxy(const FVector2d& Coordinate)
-{
-	return FVector3d(1.0, Coordinate.X, Coordinate.Y);
-}
-
 /**
  * Convert chromaticity coordinate and luminance to CIE XYZ tristimulus values.
  *
@@ -73,18 +51,6 @@ inline FVector3d XYZToxyY(const FVector3d& XYZ)
 	);
 }
 
-UE_DEPRECATED(5.3, "ToXYZ has been replaced by xyYToXYZ.")
-inline FVector3d ToXYZ(double LuminanceY, const FVector2d& Coordinate)
-{
-	return xyYToXYZ(FVector3d(Coordinate.X, Coordinate.Y, LuminanceY));
-}
-
-UE_DEPRECATED(5.3, "ToXYZ has been replaced by xyYToXYZ.")
-inline FVector3d ToXYZ(const FVector2d& Coordinate)
-{
-	return xyYToXYZ(FVector3d(Coordinate.X, Coordinate.Y, 1.0));
-}
-
 /**
  * Convenience function to get the transposed matrix, i.e. for pre-multiplied shader matrices.
  *
diff --git a/Engine/Source/Runtime/ColorManagement/Public/TransferFunctions.h b/Engine/Source/Runtime/ColorManagement/Public/TransferFunctions.h
--- a/Engine/Source/Runtime/ColorManagement/Public/TransferFunctions.h
+++ b/Engine/Source/Runtime/ColorManagement/Public/TransferFunctions.h
@@ -11,7 +11,7 @@
 
 namespace UE { namespace Color {
 
-FORCEINLINE float Linear(float Value)
+inline float Linear(float Value)
 {
 	return Value;
 }
@@ -21,7 +21,7 @@ FORCEINLINE float Linear(float Value)
 *
 * @return float encoded value.
 */
-FORCEINLINE float EncodeSRGB(float Value)
+inline float EncodeSRGB(float Value)
 {
 	if (Value <= 0.04045f / 12.92f)
 	{
@@ -38,7 +38,7 @@ FORCEINLINE float EncodeSRGB(float Value)
 *
 * @return float decoded value.
 */
-FORCEINLINE float DecodeSRGB(float Value)
+inline float DecodeSRGB(float Value)
 {
 	if (Value <= 0.04045f)
 	{
@@ -55,7 +55,7 @@ FORCEINLINE float DecodeSRGB(float Value)
 *
 * @return float encoded value.
 */
-FORCEINLINE float EncodeST2084(float Value)
+inline float EncodeST2084(float Value)
 {
 	const float Lp = 10000.0f;
 	const float m1 = 2610 / 4096.0f * (1.0f / 4.0f);
@@ -73,7 +73,7 @@ FORCEINLINE float EncodeST2084(float Value)
 *
 * @return float decoded value.
 */
-FORCEINLINE float DecodeST2084(float Value)
+inline float DecodeST2084(float Value)
 {
 	const float Lp = 10000.0f;
 	const float m1 = 2610 / 4096.0f * (1.0f / 4.0f);
@@ -92,7 +92,7 @@ FORCEINLINE float DecodeST2084(float Value)
 *
 * @return float encoded value.
 */
-FORCEINLINE float EncodeGamma22(float Value)
+inline float EncodeGamma22(float Value)
 {
 	return FGenericPlatformMath::Pow(Value, 1.0f / 2.2f);
 }
@@ -103,7 +103,7 @@ FORCEINLINE float EncodeGamma22(float Value)
 *
 * @return float decoded value.
 */
-FORCEINLINE float DecodeGamma22(float Value)
+inline float DecodeGamma22(float Value)
 {
 	return FGenericPlatformMath::Pow(Value, 2.2f);
 }
@@ -113,7 +113,7 @@ FORCEINLINE float DecodeGamma22(float Value)
 *
 * @return float encoded value.
 */
-FORCEINLINE float EncodeBT1886(float Value)
+inline float EncodeBT1886(float Value)
 {
 	const float L_B = 0;
 	const float L_W = 1;
@@ -130,7 +130,7 @@ FORCEINLINE float EncodeBT1886(float Value)
 *
 * @return float encoded value.
 */
-FORCEINLINE float EncodeGamma26(float Value)
+inline float EncodeGamma26(float Value)
 {
 	return FGenericPlatformMath::Pow(Value, 1.0f / 2.6f);
 }
@@ -141,7 +141,7 @@ FORCEINLINE float EncodeGamma26(float Value)
 *
 * @return float decoded value.
 */
-FORCEINLINE float DecodeGamma26(float Value)
+inline float DecodeGamma26(float Value)
 {
 	return FGenericPlatformMath::Pow(Value, 2.6f);
 }
@@ -151,7 +151,7 @@ FORCEINLINE float DecodeGamma26(float Value)
 *
 * @return float decoded value.
 */
-FORCEINLINE float DecodeBT1886(float Value)
+inline float DecodeBT1886(float Value)
 {
 	const float L_B = 0;
 	const float L_W = 1;
@@ -168,7 +168,7 @@ FORCEINLINE float DecodeBT1886(float Value)
 *
 * @return float encoded value.
 */
-FORCEINLINE float EncodeCineon(float Value)
+inline float EncodeCineon(float Value)
 {
 	const float BlackOffset = FGenericPlatformMath::Pow(10.0f, (95.0f - 685.0f) / 300.0f);
 	return (685.0f + 300.0f * FGenericPlatformMath::LogX(10.0f, Value * (1.0f - BlackOffset) + BlackOffset)) / 1023.0f;
@@ -179,7 +179,7 @@ FORCEINLINE float EncodeCineon(float Value)
 *
 * @return float decoded value.
 */
-FORCEINLINE float DecodeCineon(float Value)
+inline float DecodeCineon(float Value)
 {
 	const float BlackOffset = FGenericPlatformMath::Pow(10.0f, (95.0f - 685.0f) / 300.0f);
 	return (FGenericPlatformMath::Pow(10.0f, (1023.0f * Value - 685.0f) / 300.0f) - BlackOffset) / (1.0f - BlackOffset);
@@ -190,7 +190,7 @@ FORCEINLINE float DecodeCineon(float Value)
 *
 * @return float encoded value.
 */
-FORCEINLINE float EncodeREDLog(float Value)
+inline float EncodeREDLog(float Value)
 {
 	const float BlackOffset = FGenericPlatformMath::Pow(10.0f, (0.0f - 1023.0f) / 511.0f);
 	return (1023.0f + 511.0f * FGenericPlatformMath::LogX(10.0f, Value * (1.0f - BlackOffset) + BlackOffset)) / 1023.0f;
@@ -201,7 +201,7 @@ FORCEINLINE float EncodeREDLog(float Value)
 *
 * @return float decoded value.
 */
-FORCEINLINE float DecodeREDLog(float Value)
+inline float DecodeREDLog(float Value)
 {
 	const float BlackOffset = FGenericPlatformMath::Pow(10.0f, (0.0f - 1023.0f) / 511.0f);
 	return (FGenericPlatformMath::Pow(10.0f, (1023.0f * Value - 1023.0f) / 511.0f) - BlackOffset) / (1.0f - BlackOffset);
@@ -212,7 +212,7 @@ FORCEINLINE float DecodeREDLog(float Value)
 *
 * @return float encoded value.
 */
-FORCEINLINE float EncodeREDLog3G10(float Value)
+inline float EncodeREDLog3G10(float Value)
 {
 	const float A = 0.224282f;
 	const float B = 155.975327f;
@@ -236,7 +236,7 @@ FORCEINLINE float EncodeREDLog3G10(float Value)
 *
 * @return float decoded value.
 */
-FORCEINLINE float DecodeREDLog3G10(float Value)
+inline float DecodeREDLog3G10(float Value)
 {
 	const float A = 0.224282f;
 	const float B = 155.975327f;
@@ -260,7 +260,7 @@ FORCEINLINE float DecodeREDLog3G10(float Value)
 *
 * @return float encoded value.
 */
-FORCEINLINE float EncodeSLog1(float Value)
+inline float EncodeSLog1(float Value)
 {
 	Value /= 0.9f;
 	Value = 0.432699f * FGenericPlatformMath::LogX(10.0f, Value + 0.037584f) + 0.616596f + 0.03f;
@@ -272,7 +272,7 @@ FORCEINLINE float EncodeSLog1(float Value)
 *
 * @return float decoded value.
 */
-FORCEINLINE float DecodeSLog1(float Value)
+inline float DecodeSLog1(float Value)
 {
 	Value = ((Value * 1023.f) / 4.0f - 16.0f) / 219.0f;
 	Value = FGenericPlatformMath::Pow(10.0f, (Value - 0.616596f - 0.03f) / 0.432699f) - 0.037584f;
@@ -284,7 +284,7 @@ FORCEINLINE float DecodeSLog1(float Value)
 *
 * @return float encoded value.
 */
-FORCEINLINE float EncodeSLog2(float Value)
+inline float EncodeSLog2(float Value)
 {
 	if (Value >= 0.0f)
 	{
@@ -301,7 +301,7 @@ FORCEINLINE float EncodeSLog2(float Value)
 *
 * @return float decoded value.
 */
-FORCEINLINE float DecodeSLog2(float Value)
+inline float DecodeSLog2(float Value)
 {
 	if (Value >= (64.f + 0.030001222851889303f * 876.f) / 1023.f)
 	{
@@ -318,7 +318,7 @@ FORCEINLINE float DecodeSLog2(float Value)
 *
 * @return float encoded value.
 */
-FORCEINLINE float EncodeSLog3(float Value)
+inline float EncodeSLog3(float Value)
 {
 	if (Value >= 0.01125000f)
 	{
@@ -335,7 +335,7 @@ FORCEINLINE float EncodeSLog3(float Value)
 *
 * @return float decoded value.
 */
-FORCEINLINE float DecodeSLog3(float Value)
+inline float DecodeSLog3(float Value)
 {
 	if (Value >= 171.2102946929f / 1023.0f)
 	{
@@ -352,7 +352,7 @@ FORCEINLINE float DecodeSLog3(float Value)
 *
 * @return float encoded value.
 */
-FORCEINLINE float EncodeArriAlexaV3LogC(float Value)
+inline float EncodeArriAlexaV3LogC(float Value)
 {
 	const float cut = 0.010591f;
 	const float a = 5.555556f;
@@ -377,7 +377,7 @@ FORCEINLINE float EncodeArriAlexaV3LogC(float Value)
 *
 * @return float decoded value.
 */
-FORCEINLINE float DecodeArriAlexaV3LogC(float Value)
+inline float DecodeArriAlexaV3LogC(float Value)
 {
 	const float cut = 0.010591f;
 	const float a = 5.555556f;
@@ -402,7 +402,7 @@ FORCEINLINE float DecodeArriAlexaV3LogC(float Value)
 *
 * @return float encoded value.
 */
-FORCEINLINE float EncodeCanonLog(float Value)
+inline float EncodeCanonLog(float Value)
 {
 	if (Value < 0.0f)
 	{
@@ -420,7 +420,7 @@ FORCEINLINE float EncodeCanonLog(float Value)
 *
 * @return float decoded value.
 */
-FORCEINLINE float DecodeCanonLog(float Value)
+inline float DecodeCanonLog(float Value)
 {
 	if (Value < 0.0730597f)
 	{
@@ -437,7 +437,7 @@ FORCEINLINE float DecodeCanonLog(float Value)
 *
 * @return float encoded value.
 */
-FORCEINLINE float EncodeGoProProTune(float Value)
+inline float EncodeGoProProTune(float Value)
 {
 	return FGenericPlatformMath::Loge(Value * 112.f + 1.0f) / FGenericPlatformMath::Loge(113.0f);
 }
@@ -447,7 +447,7 @@ FORCEINLINE float EncodeGoProProTune(float Value)
 *
 * @return float decoded value.
 */
-FORCEINLINE float DecodeGoProProTune(float Value)
+inline float DecodeGoProProTune(float Value)
 {
 	return (FGenericPlatformMath::Pow(113.f, Value) - 1.0f) / 112.f;
 }
@@ -457,7 +457,7 @@ FORCEINLINE float DecodeGoProProTune(float Value)
 *
 * @return float encoded value.
 */
-FORCEINLINE float EncodePanasonicVLog(float Value)
+inline float EncodePanasonicVLog(float Value)
 {
 	const float b = 0.00873f;
 	const float c = 0.241514f;
@@ -478,7 +478,7 @@ FORCEINLINE float EncodePanasonicVLog(float Value)
 *
 * @return float decoded value.
 */
-FORCEINLINE float DecodePanasonicVLog(float Value)
+inline float DecodePanasonicVLog(float Value)
 {
 	const float b = 0.00873f;
 	const float c = 0.241514f;
@@ -500,10 +500,10 @@ COLORMANAGEMENT_API TFunction<float(float)> GetEncodeFunction(EEncoding Encoding
 /** Get the decode function that matches the encoding type. */
 COLORMANAGEMENT_API TFunction<float(float)> GetDecodeFunction(EEncoding Encoding);
 
-/** Encode a value based on the specified encoding type. */
+/** Encode a value based on the specified encoding type. Note: Less optimal due to function pointer call. */
 COLORMANAGEMENT_API float Encode(EEncoding Encoding, float Value);
 
-/** Decode a value based on the specified encoding type. */
+/** Decode a value based on the specified encoding type. Note: Less optimal due to function pointer call. */
 COLORMANAGEMENT_API float Decode(EEncoding Encoding, float Value);
 
 /** Get the encode function that matches the encoding type. */
@@ -512,10 +512,10 @@ COLORMANAGEMENT_API TFunction<FLinearColor(const FLinearColor&)> GetColorEncodeF
 /** Get the decode function that matches the encoding type. */
 COLORMANAGEMENT_API TFunction<FLinearColor(const FLinearColor&)> GetColorDecodeFunction(EEncoding Encoding);
 
-/** Encode a color based on the specified encoding type. */
+/** Encode a color based on the specified encoding type. Note: Less optimal due to function pointer call. */
 COLORMANAGEMENT_API FLinearColor Encode(EEncoding Encoding, const FLinearColor& Color);
 
-/** Decode a color based on the specified encoding type. */
+/** Decode a color based on the specified encoding type. Note: Less optimal due to function pointer call. */
 COLORMANAGEMENT_API FLinearColor Decode(EEncoding Encoding, const FLinearColor& Color);
 
 } } // end namespace UE::Color
