ColorManagement: Moved ColorSpace and TransferFunctions logic into into Core module.

This is needed in order to use the working color space during the luminance calculation of any color, and avoid circular dependencies between the ColorManagement and Core modules. Header deprecation and module dependency removal will be done in a subsequent CL.

diff --git a/Engine/Source/Runtime/ColorManagement/ColorManagement.Build.cs b/Engine/Source/Runtime/ColorManagement/ColorManagement.Build.cs
--- a/Engine/Source/Runtime/ColorManagement/ColorManagement.Build.cs
+++ b/Engine/Source/Runtime/ColorManagement/ColorManagement.Build.cs
@@ -6,7 +6,7 @@ public class ColorManagement : ModuleRules
 {
 	public ColorManagement(ReadOnlyTargetRules Target) : base(Target)
 	{
-		PrivateDependencyModuleNames.AddRange(
+		PublicDependencyModuleNames.AddRange(
 			new string[] {
 				"Core",
 			}
diff --git a/Engine/Source/Runtime/ColorManagement/Private/ColorSpace.cpp b/Engine/Source/Runtime/ColorManagement/Private/ColorSpace.cpp
deleted file mode 100644
diff --git a/Engine/Source/Runtime/ColorManagement/Private/Tests/ColorManagementTests.cpp b/Engine/Source/Runtime/ColorManagement/Private/Tests/ColorManagementTests.cpp
deleted file mode 100644
diff --git a/Engine/Source/Runtime/ColorManagement/Private/TransferFunctions.cpp b/Engine/Source/Runtime/ColorManagement/Private/TransferFunctions.cpp
deleted file mode 100644
diff --git a/Engine/Source/Runtime/ColorManagement/Public/ColorManagementDefines.h b/Engine/Source/Runtime/ColorManagement/Public/ColorManagementDefines.h
--- a/Engine/Source/Runtime/ColorManagement/Public/ColorManagementDefines.h
+++ b/Engine/Source/Runtime/ColorManagement/Public/ColorManagementDefines.h
@@ -2,86 +2,13 @@
 
 #pragma once
 
-#include "HAL/Platform.h"
+//TODO:Uncomment and deprecate this header a in subsequence change
+// 
+//#ifdef _MSC_VER
+//#pragma message(__FILE__"(9): warning: use ColorManagement/ColorManagementDefines.h from the Core module instead of ColorManagementDefines.h")
+//#else
+//#pragma message("#include use ColorManagement/ColorManagementDefines.h from the Core module instead of ColorManagementDefines.h")
+//#endif
 
-namespace UE { namespace Color {
+#include "ColorManagement/ColorManagementDefines.h"
 
-/** Increment upon breaking changes to the EEncoding enum. Note that changing this forces a rebuild of textures that rely on it. */
-constexpr uint32 ENCODING_TYPES_VER = 4;
-
-/** Increment upon breaking changes to the EColorSpace and EChromaticAdaptationMethod enums. Note that changing this forces a rebuild of textures that rely on it.*/
-constexpr uint32 COLORSPACE_VER = 2;
-
-/** List of available encodings/transfer functions.
-* 
-* NOTE: This list is replicated as a UENUM in TextureDefines.h, and both should always match.
-*/
-enum class EEncoding : uint8
-{
-	None = 0,
-	Linear = 1,
-	sRGB,
-	ST2084,
-	Gamma22,
-	BT1886,
-	Gamma26,
-	Cineon,
-	REDLog,
-	REDLog3G10,
-	SLog1,
-	SLog2,
-	SLog3,
-	AlexaV3LogC,
-	CanonLog,
-	ProTune,
-	VLog,
-	Max,
-};
-
-/** List of available color spaces. (Increment COLORSPACE_VER upon breaking changes to the list.)
-* 
-* NOTE: This list is partially replicated as a UENUM in TextureDefines.h: any type exposed to textures should match the enum value below.
-*/
-enum class EColorSpace : uint8
-{
-	None = 0,
-	sRGB = 1,
-	Rec2020 = 2,
-	ACESAP0 = 3,
-	ACESAP1 = 4,
-	P3DCI = 5,
-	P3D65 = 6,
-	REDWideGamut = 7,
-	SonySGamut3 = 8,
-	SonySGamut3Cine = 9,
-	AlexaWideGamut = 10,
-	CanonCinemaGamut = 11,
-	GoProProtuneNative = 12,
-	PanasonicVGamut = 13,
-	PLASA_E1_54 = 14,
-	Max,
-};
-
-
-/** List of available chromatic adaptation methods.
-* 
-* NOTE: This list is replicated as a UENUM in TextureDefines.h, and both should always match.
-*/
-enum class EChromaticAdaptationMethod : uint8
-{
-	None = 0,
-	Bradford = 1,
-	CAT02 = 2,
-	Max,
-};
-
-/** List of standard white points. */
-enum class EWhitePoint : uint8
-{
-	CIE1931_D65 = 0,
-	ACES_D60 = 1,
-	DCI_CalibrationWhite,
-	Max,
-};
-
-} } // end namespace UE::Color
diff --git a/Engine/Source/Runtime/ColorManagement/Public/ColorSpace.h b/Engine/Source/Runtime/ColorManagement/Public/ColorSpace.h
--- a/Engine/Source/Runtime/ColorManagement/Public/ColorSpace.h
+++ b/Engine/Source/Runtime/ColorManagement/Public/ColorSpace.h
@@ -2,365 +2,13 @@
 
 #pragma once
 
-#include "ColorManagementDefines.h"
-#include "Containers/StaticArray.h"
-#include "CoreTypes.h"
-#include "Math/Color.h"
-#include "Math/Matrix.h"
-#include "Math/UnrealMathSSE.h"
-#include "Math/Vector.h"
-#include "Misc/AssertionMacros.h"
-#include "Serialization/Archive.h"
-#include "Templates/UnrealTypeTraits.h"
+//TODO:Uncomment and deprecate this header a in subsequence change
+// 
+//#ifdef _MSC_VER
+//#pragma message(__FILE__"(9): warning: use ColorManagement/ColorSpace.h from the Core module instead of ColorSpace.h")
+//#else
+//#pragma message("#include use ColorManagement/ColorSpace.h from the Core module instead of ColorSpace.h")
+//#endif
 
-namespace UE { namespace Color {
+#include "ColorManagement/ColorSpace.h"
 
-/**
- * Convert chromaticity coordinate and luminance to CIE XYZ tristimulus values.
- *
- * @return FVector3d
- */
-inline FVector3d xyYToXYZ(const FVector3d& xyY)
-{
-	double Divisor = FMath::Max(xyY[1], 1e-10);
-
-	return FVector3d(
-		xyY[0] * xyY[2] / Divisor,
-		xyY[2],
-		(1.0 - xyY[0] - xyY[1]) * xyY[2] / Divisor
-	);
-}
-
-/**
- * Convert CIE XYZ tristimulus values to chromaticitiy coordinates and luminance.
- *
- * @return FVector3d
- */
-inline FVector3d XYZToxyY(const FVector3d& XYZ)
-{
-	double Divisor = XYZ[0] + XYZ[1] + XYZ[2];
-	if (Divisor == 0.0)
-	{
-		Divisor = 1e-10;
-	}
-
-	return FVector3d(
-		XYZ[0] / Divisor,
-		XYZ[1] / Divisor,
-		XYZ[1]
-	);
-}
-
-/**
- * Convenience function to get the transposed matrix, i.e. for pre-multiplied shader matrices.
- *
- * @return TMatrix<T>
- */
-template<typename T>
-inline UE::Math::TMatrix<T> Transpose(const FMatrix44d& Transform)
-{
-	if constexpr (std::is_same_v<T, double>)
-	{
-		return Transform.GetTransposed();
-	}
-	else
-	{
-		return UE::Math::TMatrix<T>(Transform).GetTransposed();
-	}
-}
-
-/**
- * Get standard white point coordinates.
- *
- * @param InWhitePoint White point type.
- * @return FVector2d Chromaticity coordinates.
- */
-inline FVector2d GetWhitePoint(EWhitePoint InWhitePoint)
-{
-	switch (InWhitePoint)
-	{
-	case UE::Color::EWhitePoint::CIE1931_D65:
-		return FVector2d(0.3127, 0.3290);
-	case UE::Color::EWhitePoint::ACES_D60:
-		return FVector2d(0.32168, 0.33767);
-	case UE::Color::EWhitePoint::DCI_CalibrationWhite:
-		return FVector2d(0.314, 0.351);
-
-	default:
-		checkNoEntry();
-		return FVector2d();
-	}
-}
-
-
-/** Color space definition as 4 chromaticity coordinates, in double precision internally. */
-class FColorSpace
-{
-public:
-	/**
-	 * Get the global engine working color space (as a singleton).
-	 * 
-	 * @return FColorSpace working color space
-	 */
-	static COLORMANAGEMENT_API const FColorSpace& GetWorking();
-
-	/**
-	 * Set the global engine working color space (as a singleton).
-	 * 
-	 * @param ColorSpace working color space
-	 */
-	static COLORMANAGEMENT_API void SetWorking(FColorSpace ColorSpace);
-
-
-	/** Constructor */
-	FColorSpace() {}
-
-	/**
-	* Constructor
-	*
-	* @param InRed Chromaticity 2D coordinates for the red color.
-	* @param InGreen Chromaticity 2D coordinates for the green color.
-	* @param InBlue Chromaticity 2D coordinates for the blue color.
-	* @param InWhite Chromaticity 2D coordinates for the white point.
-	*/
-	COLORMANAGEMENT_API explicit FColorSpace(const FVector2d& InRed, const FVector2d& InGreen, const FVector2d& InBlue, const FVector2d& InWhite);
-
-	/**
-	* Constructor
-	*
-	* @param ColorSpaceType Color space type.
-	*/
-	COLORMANAGEMENT_API explicit FColorSpace(UE::Color::EColorSpace ColorSpaceType);
-
-	FColorSpace(FColorSpace&&) = default;
-	FColorSpace(const FColorSpace&) = default;
-	FColorSpace& operator=(FColorSpace&&) = default;
-	FColorSpace& operator=(const FColorSpace&) = default;
-
-	/**
-	* Make the chromaticities of the color space type.
-	* 
-	* @return TStaticArray of four chromaticities
-	*/
-	COLORMANAGEMENT_API static TStaticArray<FVector2d, 4> MakeChromaticities(UE::Color::EColorSpace ColorSpaceType);
-
-	/**
-	* Getter for the color space chromaticity coordinates.
-	* 
-	* @param OutRed FVector2d for the red color chromaticity coordinate.
-	* @param OutGreen FVector2d for the green color chromaticity coordinate.
-	* @param OutBlue FVector2d for the blue color chromaticity coordinate.
-	* @param OutWhite FVector2d for the white color chromaticity coordinate.
-	*/
-	template<typename T>
-	inline void GetChromaticities(UE::Math::TVector2<T>& OutRed, UE::Math::TVector2<T>& OutGreen, UE::Math::TVector2<T>& OutBlue, UE::Math::TVector2<T>& OutWhite) const
-	{
-		OutRed		= Chromaticities[0];
-		OutGreen	= Chromaticities[1];
-		OutBlue		= Chromaticities[2];
-		OutWhite	= Chromaticities[3];
-	}
-
-	/**
-	* Gets the color space's red chromaticity coordinates.
-	*
-	* @return FVector2d xy coordinates.
-	*/
-	inline const FVector2d& GetRedChromaticity() const
-	{
-		return Chromaticities[0];
-	}
-
-	/**
-	* Gets the color space's green chromaticity coordinates.
-	*
-	* @return FVector2d xy coordinates.
-	*/
-	inline const FVector2d& GetGreenChromaticity() const
-	{
-		return Chromaticities[1];
-	}
-
-	/**
-	* Gets the color space's blue chromaticity coordinates.
-	*
-	* @return FVector2d xy coordinates.
-	*/
-	inline const FVector2d& GetBlueChromaticity() const
-	{
-		return Chromaticities[2];
-	}
-
-	/**
-	* Gets the color space's white point chromaticity coordinates.
-	*
-	* @return FVector2d xy coordinates.
-	*/
-	inline const FVector2d& GetWhiteChromaticity() const
-	{
-		return Chromaticities[3];
-	}
-
-	/**
-	* Gets the RGB-to-XYZ conversion matrix.
-	*
-	* @return FMatrix conversion matrix.
-	*/
-	inline const FMatrix44d& GetRgbToXYZ() const
-	{
-		return RgbToXYZ;
-	}
-
-	/**
-	* Gets the XYZ-to-RGB conversion matrix.
-	*
-	* @return FMatrix conversion matrix.
-	*/
-	inline const FMatrix44d& GetXYZToRgb() const
-	{
-		return XYZToRgb;
-	}
-
-	/**
-	 * Check against another color space for equality.
-	 *
-	 * @param ColorSpace The vector to check against.
-	 * @return true if the color spaces are equal, false otherwise.
-	 */
-	inline bool operator==(const FColorSpace& ColorSpace) const
-	{
-		return Chromaticities == ColorSpace.Chromaticities;
-	}
-
-	/**
-	 * Check against another color space for inequality.
-	 *
-	 * @param ColorSpace The color space to check against.
-	 * @return true if the color spaces are not equal, false otherwise.
-	 */
-	inline bool operator!=(const FColorSpace& ColorSpace) const
-	{
-		return Chromaticities != ColorSpace.Chromaticities;
-	}
-
-	/**
-	 * Check against another color space for equality, within specified error limits.
-	 *
-	 * @param ColorSpace The color space to check against.
-	 * @param Tolerance Error tolerance.
-	 * @return true if the color spaces are equal within tolerance limits, false otherwise.
-	 */
-	COLORMANAGEMENT_API bool Equals(const FColorSpace& ColorSpace, double Tolerance = 1.e-7) const;
-
-
-	/**
-	 * Check against color space chromaticities for equality, within specified error limits.
-	 *
-	 * @param InChromaticities The color space chromaticities to check against.
-	 * @param Tolerance Error tolerance.
-	 * @return true if the color spaces are equal within tolerance limits, false otherwise.
-	 */
-	COLORMANAGEMENT_API bool Equals(const TStaticArray<FVector2d, 4>& InChromaticities, double Tolerance = 1.e-7) const;
-
-	/**
-	 * Convenience function to verify if the color space matches the engine's default sRGB chromaticities.
-	 *
-	 * @return true if sRGB.
-	 */
-	COLORMANAGEMENT_API bool IsSRGB() const;
-
-	/**
-	* Converts temperature in Kelvins of a black body radiator to an RGB color in the current space.
-	*/
-	COLORMANAGEMENT_API FLinearColor MakeFromColorTemperature(float Temp) const;
-
-	/**
-	 * Calculate the color's luminance value in the current space.
-	 *
-	 * @param Color The sampled color.
-	 * @return float Luminance value.
-	 */
-	COLORMANAGEMENT_API float GetLuminance(const FLinearColor& Color) const;
-
-	/**
-	 * Get the luminance factors in the current space.
-	 *
-	 * @return FLinearColor Luminance factors.
-	 */
-	COLORMANAGEMENT_API FLinearColor GetLuminanceFactors() const;
-
-private:
-
-	COLORMANAGEMENT_API FMatrix44d CalcRgbToXYZ() const;
-
-	/** Red, green, blue, white chromaticities, in order. */
-	TStaticArray<FVector2d, 4> Chromaticities;
-
-	FMatrix44d RgbToXYZ;
-	FMatrix44d XYZToRgb;
-
-	bool bIsSRGB;
-
-public:
-
-	/**
-	 * Serializer.
-	 *
-	 * @param Ar The Serialization Archive.
-	 * @param CS The Color Space being serialized.
-	 * @return Reference to the Archive after serialization.
-	 */
-	friend FArchive& operator<<(FArchive& Ar, FColorSpace& CS)
-	{
-		return Ar << CS.Chromaticities;
-	}
-
-	bool Serialize(FArchive& Ar)
-	{
-		Ar << *this;
-		return true;
-	}
-};
-
-
-struct FColorSpaceTransform : FMatrix44d
-{
-	/**
-	* Constructor: create a color space transformation matrix from a source to a target color space using the RGB-XYZ-RGB conversions.
-	*
-	* @param Source Source color space.
-	* @param Target Target color space.
-	* @param Method Chromatic adapation method.
-	*/
-	COLORMANAGEMENT_API explicit FColorSpaceTransform(const FColorSpace& Src, const FColorSpace& Dst, EChromaticAdaptationMethod Method = EChromaticAdaptationMethod::Bradford);
-
-	/**
-	* Constructor: create a color space transformation matrix from a raw matrix.
-	*
-	* @param Matrix Color space transformation matrix.
-	*/
-	COLORMANAGEMENT_API explicit FColorSpaceTransform(FMatrix44d Matrix);
-
-	/**
-	* Apply color space transform to FLinearColor.
-	* 
-	* @param Color Color to transform.
-	*/
-	COLORMANAGEMENT_API FLinearColor Apply(const FLinearColor& Color) const;
-
-	/**
-	* Calculate the chromatic adaptation matrix using the specific method.
-	*
-	* @param SourceXYZ Source color in XYZ space.
-	* @param TargetXYZ Target color in XYZ space.
-	* @param Method Adaptation method (None, Bradford, CAT02).
-	*/
-	COLORMANAGEMENT_API static FMatrix44d CalcChromaticAdaptionMatrix(FVector3d SourceXYZ, FVector3d TargetXYZ, EChromaticAdaptationMethod Method = EChromaticAdaptationMethod::Bradford);
-
-	/**
-	* Convenience function to get a (statically cached) conversion from sRGB/Rec709 to the working color space.
-	*/
-	COLORMANAGEMENT_API static FColorSpaceTransform GetSRGBToWorkingColorSpace();
-};
-
-} }  // end namespace UE::Color
diff --git a/Engine/Source/Runtime/ColorManagement/Public/TransferFunctions.h b/Engine/Source/Runtime/ColorManagement/Public/TransferFunctions.h
--- a/Engine/Source/Runtime/ColorManagement/Public/TransferFunctions.h
+++ b/Engine/Source/Runtime/ColorManagement/Public/TransferFunctions.h
@@ -2,520 +2,12 @@
 
 #pragma once
 
-#include "ColorManagementDefines.h"
-#include "CoreTypes.h"
-#include "GenericPlatform/GenericPlatformMath.h"
-#include "Math/Color.h"
-#include "Math/UnrealMathUtility.h"
-#include "Templates/Function.h"
-
-namespace UE { namespace Color {
-
-inline float Linear(float Value)
-{
-	return Value;
-}
-
-/**
-* Encode value to sRGB.
-*
-* @return float encoded value.
-*/
-inline float EncodeSRGB(float Value)
-{
-	if (Value <= 0.04045f / 12.92f)
-	{
-		return Value * 12.92f;
-	}
-	else
-	{
-		return  FGenericPlatformMath::Pow(Value, (1.0f / 2.4f)) * 1.055f - 0.055f;
-	}
-}
-
-/**
-* Decode value with an sRGB encoding.
-*
-* @return float decoded value.
-*/
-inline float DecodeSRGB(float Value)
-{
-	if (Value <= 0.04045f)
-	{
-		return Value / 12.92f;
-	}
-	else
-	{
-		return  FGenericPlatformMath::Pow((Value + 0.055f) / 1.055f, 2.4f);
-	}
-}
-
-/**
-* Encode value to SMPTE ST 2084:2014.
-*
-* @return float encoded value.
-*/
-inline float EncodeST2084(float Value)
-{
-	const float Lp = 10000.0f;
-	const float m1 = 2610 / 4096.0f * (1.0f / 4.0f);
-	const float m2 = 2523 / 4096.0f * 128.0f;
-	const float c1 = 3424 / 4096.0f;
-	const float c2 = 2413 / 4096.0f * 32.f;
-	const float c3 = 2392 / 4096.0f * 32.f;
-
-	Value = FGenericPlatformMath::Pow(Value / Lp, m1);
-	return FGenericPlatformMath::Pow((c1 + c2 * Value) / (c3 * Value + 1), m2);
-}
-
-/**
-* Decode value with a SMPTE ST 2084:2014 encoding.
-*
-* @return float decoded value.
-*/
-inline float DecodeST2084(float Value)
-{
-	const float Lp = 10000.0f;
-	const float m1 = 2610 / 4096.0f * (1.0f / 4.0f);
-	const float m2 = 2523 / 4096.0f * 128.0f;
-	const float c1 = 3424 / 4096.0f;
-	const float c2 = 2413 / 4096.0f * 32.f;
-	const float c3 = 2392 / 4096.0f * 32.f;
-
-	const float Vp = FGenericPlatformMath::Pow(Value, 1.0f / m2);
-	Value = FGenericPlatformMath::Max(0.0f, Vp - c1);
-	return FGenericPlatformMath::Pow((Value / (c2 - c3 * Vp)), 1.0f / m1) * Lp;
-}
-
-/**
-* Encode value to Gamma 2.2.
-*
-* @return float encoded value.
-*/
-inline float EncodeGamma22(float Value)
-{
-	return FGenericPlatformMath::Pow(Value, 1.0f / 2.2f);
-}
-
-
-/**
-* Decode value with a Gamma 2.2 encoding.
-*
-* @return float decoded value.
-*/
-inline float DecodeGamma22(float Value)
-{
-	return FGenericPlatformMath::Pow(Value, 2.2f);
-}
-
-/**
-* Encode value to ITU-R BT.1886.
-*
-* @return float encoded value.
-*/
-inline float EncodeBT1886(float Value)
-{
-	const float L_B = 0;
-	const float L_W = 1;
-	float Gamma = 2.40f;
-	float GammaInv = 1.0f / Gamma;
-	float N = FGenericPlatformMath::Pow(L_W, GammaInv) - FGenericPlatformMath::Pow(L_B, GammaInv);
-	float A = FGenericPlatformMath::Pow(N, Gamma);
-	float B = FGenericPlatformMath::Pow(L_B, GammaInv) / N;
-	return FGenericPlatformMath::Pow(Value / A, GammaInv) - B;
-}
-
-/**
-* Encode value to Gamma 2.6.
-*
-* @return float encoded value.
-*/
-inline float EncodeGamma26(float Value)
-{
-	return FGenericPlatformMath::Pow(Value, 1.0f / 2.6f);
-}
-
-
-/**
-* Decode value with a Gamma 2.6 encoding.
-*
-* @return float decoded value.
-*/
-inline float DecodeGamma26(float Value)
-{
-	return FGenericPlatformMath::Pow(Value, 2.6f);
-}
-
-/**
-* Decode value with an ITU-R BT.1886 encoding.
-*
-* @return float decoded value.
-*/
-inline float DecodeBT1886(float Value)
-{
-	const float L_B = 0;
-	const float L_W = 1;
-	float Gamma = 2.40f;
-	float GammaInv = 1.0f / Gamma;
-	float N = FGenericPlatformMath::Pow(L_W, GammaInv) - FGenericPlatformMath::Pow(L_B, GammaInv);
-	float A = FGenericPlatformMath::Pow(N, Gamma);
-	float B = FGenericPlatformMath::Pow(L_B, GammaInv) / N;
-	return A * FGenericPlatformMath::Pow(FGenericPlatformMath::Max(Value + B, 0.0f), Gamma);
-}
-
-/**
-* Encode value to Cineon.
-*
-* @return float encoded value.
-*/
-inline float EncodeCineon(float Value)
-{
-	const float BlackOffset = FGenericPlatformMath::Pow(10.0f, (95.0f - 685.0f) / 300.0f);
-	return (685.0f + 300.0f * FGenericPlatformMath::LogX(10.0f, Value * (1.0f - BlackOffset) + BlackOffset)) / 1023.0f;
-}
-
-/**
-* Decode value with a Cineon encoding.
-*
-* @return float decoded value.
-*/
-inline float DecodeCineon(float Value)
-{
-	const float BlackOffset = FGenericPlatformMath::Pow(10.0f, (95.0f - 685.0f) / 300.0f);
-	return (FGenericPlatformMath::Pow(10.0f, (1023.0f * Value - 685.0f) / 300.0f) - BlackOffset) / (1.0f - BlackOffset);
-}
-
-/**
-* Encode value to RED Log.
-*
-* @return float encoded value.
-*/
-inline float EncodeREDLog(float Value)
-{
-	const float BlackOffset = FGenericPlatformMath::Pow(10.0f, (0.0f - 1023.0f) / 511.0f);
-	return (1023.0f + 511.0f * FGenericPlatformMath::LogX(10.0f, Value * (1.0f - BlackOffset) + BlackOffset)) / 1023.0f;
-}
-
-/**
-* Decode value with a RED Log encoding.
-*
-* @return float decoded value.
-*/
-inline float DecodeREDLog(float Value)
-{
-	const float BlackOffset = FGenericPlatformMath::Pow(10.0f, (0.0f - 1023.0f) / 511.0f);
-	return (FGenericPlatformMath::Pow(10.0f, (1023.0f * Value - 1023.0f) / 511.0f) - BlackOffset) / (1.0f - BlackOffset);
-}
-
-/**
-* Encode value to RED Log3G10.
-*
-* @return float encoded value.
-*/
-inline float EncodeREDLog3G10(float Value)
-{
-	const float A = 0.224282f;
-	const float B = 155.975327f;
-	const float C = 0.01f;
-	const float G = 15.1927f;
-
-	Value += C;
-
-	if (Value < 0.0f)
-	{
-		return Value * G;
-	}
-	else
-	{
-		return FGenericPlatformMath::Sign(Value) * A * FGenericPlatformMath::LogX(10.0f, B * FGenericPlatformMath::Abs(Value) + 1.0f);
-	}
-}
-
-/**
-* Decode value with a RED Log3G10 encoding.
-*
-* @return float decoded value.
-*/
-inline float DecodeREDLog3G10(float Value)
-{
-	const float A = 0.224282f;
-	const float B = 155.975327f;
-	const float C = 0.01f;
-	const float G = 15.1927f;
-
-	if (Value < 0.0f)
-	{
-		Value /= G;
-	}
-	else
-	{
-		Value = FGenericPlatformMath::Sign(Value) * (FGenericPlatformMath::Pow(10.0f, FGenericPlatformMath::Abs(Value) / A ) - 1.0f) / B;
-	}
-
-	return Value - C;
-}
-
-/**
-* Encode value to Sony S-Log1.
-*
-* @return float encoded value.
-*/
-inline float EncodeSLog1(float Value)
-{
-	Value /= 0.9f;
-	Value = 0.432699f * FGenericPlatformMath::LogX(10.0f, Value + 0.037584f) + 0.616596f + 0.03f;
-	return (Value * 219.0f + 16.0f) * 4.0f / 1023.0f;
-}
-
-/**
-* Decode value with a Sony S-Log1 encoding.
-*
-* @return float decoded value.
-*/
-inline float DecodeSLog1(float Value)
-{
-	Value = ((Value * 1023.f) / 4.0f - 16.0f) / 219.0f;
-	Value = FGenericPlatformMath::Pow(10.0f, (Value - 0.616596f - 0.03f) / 0.432699f) - 0.037584f;
-	return Value * 0.9f;
-}
-
-/**
-* Encode value to Sony S-Log2.
-*
-* @return float encoded value.
-*/
-inline float EncodeSLog2(float Value)
-{
-	if (Value >= 0.0f)
-	{
-		return (64.0f + 876.0f * (0.432699f * FGenericPlatformMath::LogX(10.0f, 155.0f * Value / 197.1f + 0.037584f) + 0.646596f)) / 1023.f;
-	}
-	else
-	{
-		return (64.0f +876.0f * (Value * 3.53881278538813f / 0.9f) + 0.646596f + 0.030001222851889303f) / 1023.f;
-	}
-}
-
-/**
-* Decode value with a Sony S-Log2 encoding.
-*
-* @return float decoded value.
-*/
-inline float DecodeSLog2(float Value)
-{
-	if (Value >= (64.f + 0.030001222851889303f * 876.f) / 1023.f)
-	{
-		return 197.1f * (FGenericPlatformMath::Pow(10.0f, ((Value * 1023.f - 64.f) / 876.f - 0.646596f) / 0.432699f) - 0.037584f) / 155.f;
-	}
-	else
-	{
-		return 0.9f * ((Value * 1023.f - 64.f) / 876.f - 0.030001222851889303f) / 3.53881278538813f;
-	}
-}
-
-/**
-* Encode value to Sony S-Log3.
-*
-* @return float encoded value.
-*/
-inline float EncodeSLog3(float Value)
-{
-	if (Value >= 0.01125000f)
-	{
-		return (420.0f + FGenericPlatformMath::LogX(10.0f, (Value + 0.01f) / 0.19f) * 261.5f) / 1023.0f;
-	}
-	else
-	{
-		return (Value * 76.2102946929f / 0.01125f + 95.0f) / 1023.0f;
-	}
-}
-
-/**
-* Decode value with a Sony S-Log3 encoding.
-*
-* @return float decoded value.
-*/
-inline float DecodeSLog3(float Value)
-{
-	if (Value >= 171.2102946929f / 1023.0f)
-	{
-		return (FGenericPlatformMath::Pow(10.0f, (Value * 1023.0f - 420.f) / 261.5f)) * 0.19f - 0.01f;
-	}
-	else
-	{
-		return (Value * 1023.0f - 95.0f) * 0.01125000f / (171.2102946929f - 95.0f);
-	}
-}
-
-/**
-* Encode value to ARRI Alexa LogC.
-*
-* @return float encoded value.
-*/
-inline float EncodeArriAlexaV3LogC(float Value)
-{
-	const float cut = 0.010591f;
-	const float a = 5.555556f;
-	const float b = 0.052272f;
-	const float c = 0.247190f;
-	const float d = 0.385537f;
-	const float e = 5.367655f;
-	const float f = 0.092809f;
-
-	if (Value > cut)
-	{
-		return c * FGenericPlatformMath::LogX(10.0f, a * Value + b) + d;
-	}
-	else
-	{
-		return e * Value + f;
-	}
-}
-
-/**
-* Decode value with an ARRI Alexa LogC encoding.
-*
-* @return float decoded value.
-*/
-inline float DecodeArriAlexaV3LogC(float Value)
-{
-	const float cut = 0.010591f;
-	const float a = 5.555556f;
-	const float b = 0.052272f;
-	const float c = 0.247190f;
-	const float d = 0.385537f;
-	const float e = 5.367655f;
-	const float f = 0.092809f;
-	
-	if (Value > e * cut + f)
-	{
-		return (FGenericPlatformMath::Pow(10.0f, (Value - d) / c) - b) / a;
-	}
-	else
-	{
-		return (Value - f) / e;
-	}
-}
-
-/**
-* Encode value to Canon Log.
-*
-* @return float encoded value.
-*/
-inline float EncodeCanonLog(float Value)
-{
-	if (Value < 0.0f)
-	{
-		return -(0.529136f * (FGenericPlatformMath::LogX(10.0f, -Value * 10.1596f + 1.0f)) - 0.0730597f);
-	}
-	else
-	{
-		return 0.529136f * FGenericPlatformMath::LogX(10.0f, 10.1596f * Value + 1.0f) + 0.0730597f;
-	}
-}
-
-
-/**
-* Decode value with a Canon Log encoding.
-*
-* @return float decoded value.
-*/
-inline float DecodeCanonLog(float Value)
-{
-	if (Value < 0.0730597f)
-	{
-		return -(FGenericPlatformMath::Pow(10.0f, (0.0730597f - Value) / 0.529136f) - 1.0f) / 10.1596f;
-	}
-	else
-	{
-		return (FGenericPlatformMath::Pow(10.0f, (Value - 0.0730597f) / 0.529136f) - 1.0f) / 10.1596f;
-	}
-}
-
-/**
-* Encode value to GoPro ProTune.
-*
-* @return float encoded value.
-*/
-inline float EncodeGoProProTune(float Value)
-{
-	return FGenericPlatformMath::Loge(Value * 112.f + 1.0f) / FGenericPlatformMath::Loge(113.0f);
-}
-
-/**
-* Decode value with a GoPro ProTune encoding.
-*
-* @return float decoded value.
-*/
-inline float DecodeGoProProTune(float Value)
-{
-	return (FGenericPlatformMath::Pow(113.f, Value) - 1.0f) / 112.f;
-}
-
-/**
-* Encode value to Panasonic V-Log.
-*
-* @return float encoded value.
-*/
-inline float EncodePanasonicVLog(float Value)
-{
-	const float b = 0.00873f;
-	const float c = 0.241514f;
-	const float d = 0.598206f;
-
-	if( Value < 0.01f)
-	{
-		return 5.6f * Value + 0.125f;
-	}
-	else
-	{
-		return c * FGenericPlatformMath::LogX(10.0f, Value + b) + d;
-	}
-}
-
-/**
-* Decode value with a Panasonic V-Log encoding.
-*
-* @return float decoded value.
-*/
-inline float DecodePanasonicVLog(float Value)
-{
-	const float b = 0.00873f;
-	const float c = 0.241514f;
-	const float d = 0.598206f;
-
-	if (Value < 0.181f)
-	{
-		return (Value - 0.125f) / 5.6f;
-	}
-	else
-	{
-		return FGenericPlatformMath::Pow(10.0f, (Value - d) / c) - b;
-	}
-}
-
-/** Get the encode function that matches the encoding type. */
-COLORMANAGEMENT_API TFunction<float(float)> GetEncodeFunction(EEncoding Encoding);
-
-/** Get the decode function that matches the encoding type. */
-COLORMANAGEMENT_API TFunction<float(float)> GetDecodeFunction(EEncoding Encoding);
-
-/** Encode a value based on the specified encoding type. Note: Less optimal due to function pointer call. */
-COLORMANAGEMENT_API float Encode(EEncoding Encoding, float Value);
-
-/** Decode a value based on the specified encoding type. Note: Less optimal due to function pointer call. */
-COLORMANAGEMENT_API float Decode(EEncoding Encoding, float Value);
-
-/** Get the encode function that matches the encoding type. */
-COLORMANAGEMENT_API TFunction<FLinearColor(const FLinearColor&)> GetColorEncodeFunction(EEncoding Encoding);
-
-/** Get the decode function that matches the encoding type. */
-COLORMANAGEMENT_API TFunction<FLinearColor(const FLinearColor&)> GetColorDecodeFunction(EEncoding Encoding);
-
-/** Encode a color based on the specified encoding type. Note: Less optimal due to function pointer call. */
-COLORMANAGEMENT_API FLinearColor Encode(EEncoding Encoding, const FLinearColor& Color);
-
-/** Decode a color based on the specified encoding type. Note: Less optimal due to function pointer call. */
-COLORMANAGEMENT_API FLinearColor Decode(EEncoding Encoding, const FLinearColor& Color);
-
-} } // end namespace UE::Color
+//TODO: Uncomment and deprecate this header a in subsequence change
+// 
+//#ifdef _MSC_VER
+//#pragma message(__FILE__"(9): warning: use ColorManagement/TransferFunctions.h from the Core module instead of TransferFunctions.h")
+//#else
+//#pragma message("#include use ColorManagement/TransferFunctions.h from the Core module instead of TransferFunctions.h")
+//#endif
+
+#include "ColorManagement/TransferFunctions.h"
diff --git a/Engine/Source/Runtime/Core/Private/ColorManagement/ColorSpace.cpp b/Engine/Source/Runtime/Core/Private/ColorManagement/ColorSpace.cpp
new file mode 100644
--- /dev/null
+++ b/Engine/Source/Runtime/Core/Private/ColorManagement/ColorSpace.cpp
@@ -0,0 +1,352 @@
+// Copyright Epic Games, Inc. All Rights Reserved.
+
+#include "ColorManagement/ColorSpace.h"
+#include "Misc/ConfigCacheIni.h"
+#include "Math/VectorRegister.h"
+
+namespace UE { namespace Color {
+
+static bool bIsWorkingColorSpaceReadyForUse = false;
+static FColorSpace WorkingColorSpace = FColorSpace(EColorSpace::sRGB);
+
+void PreloadWorkingColorSpace()
+{
+	if (bIsWorkingColorSpaceReadyForUse)
+	{
+		return;
+	}
+
+	check(GConfig != nullptr && GConfig->IsReadyForUse());
+
+	bool bIsWorkingColorSpaceInConfig = true;
+	TStaticArray<FVector2d, 4> Chromaticities;
+	bIsWorkingColorSpaceInConfig &= GConfig->GetVector2D(TEXT("/Script/Engine.RendererSettings"), TEXT("RedChromaticityCoordinate"),   Chromaticities[0], GEngineIni);
+	bIsWorkingColorSpaceInConfig &= GConfig->GetVector2D(TEXT("/Script/Engine.RendererSettings"), TEXT("GreenChromaticityCoordinate"), Chromaticities[1], GEngineIni);
+	bIsWorkingColorSpaceInConfig &= GConfig->GetVector2D(TEXT("/Script/Engine.RendererSettings"), TEXT("BlueChromaticityCoordinate"),  Chromaticities[2], GEngineIni);
+	bIsWorkingColorSpaceInConfig &= GConfig->GetVector2D(TEXT("/Script/Engine.RendererSettings"), TEXT("WhiteChromaticityCoordinate"), Chromaticities[3], GEngineIni);
+
+	if (bIsWorkingColorSpaceInConfig)
+	{
+		FColorSpace::SetWorking(FColorSpace(Chromaticities[0], Chromaticities[1], Chromaticities[2], Chromaticities[3]));
+	}
+	else
+	{
+		// The working color space wasn't modified/serialized to the config so we keep the implicit sRGB default.
+		bIsWorkingColorSpaceReadyForUse = true;
+	}
+};
+
+const FColorSpace& FColorSpace::GetWorking()
+{
+	/**
+	 * NOTE: Addresses issue where shader compilation can request the working color space before it has been loaded by renderer settings.
+	 * We optionally early-load the settings here and leave existing singleton interface as is.
+	 */
+	UE_CALL_ONCE(PreloadWorkingColorSpace);
+
+	return WorkingColorSpace;
+}
+
+void FColorSpace::SetWorking(FColorSpace ColorSpace)
+{
+	WorkingColorSpace = MoveTemp(ColorSpace);
+	bIsWorkingColorSpaceReadyForUse = true;
+}
+
+static bool IsSRGBChromaticities(const TStaticArray<FVector2d, 4>& Chromaticities, double Tolerance = 1.e-7)
+{
+	return	Chromaticities[0].Equals(FVector2d(0.64, 0.33), Tolerance) &&
+		Chromaticities[1].Equals(FVector2d(0.30, 0.60), Tolerance) &&
+		Chromaticities[2].Equals(FVector2d(0.15, 0.06), Tolerance) &&
+		Chromaticities[3].Equals(FVector2d(0.3127, 0.3290), Tolerance);
+}
+
+FColorSpace::FColorSpace(const FVector2d& InRed, const FVector2d& InGreen, const FVector2d& InBlue, const FVector2d& InWhite)
+{
+	Chromaticities[0] = InRed;
+	Chromaticities[1] = InGreen;
+	Chromaticities[2] = InBlue;
+	Chromaticities[3] = InWhite;
+
+	bIsSRGB = IsSRGBChromaticities(Chromaticities);
+
+	RgbToXYZ = CalcRgbToXYZ();
+	XYZToRgb = RgbToXYZ.Inverse();
+}
+
+FColorSpace::FColorSpace(EColorSpace ColorSpaceType)
+	: Chromaticities(MakeChromaticities(ColorSpaceType))
+	, bIsSRGB(ColorSpaceType == EColorSpace::sRGB)
+{
+	RgbToXYZ = CalcRgbToXYZ();
+	XYZToRgb = RgbToXYZ.Inverse();
+}
+
+TStaticArray<FVector2d, 4> FColorSpace::MakeChromaticities(UE::Color::EColorSpace ColorSpaceType)
+{
+	TStaticArray<FVector2d, 4> Chromaticities(InPlace, FVector2d::Zero());
+
+	switch (ColorSpaceType)
+	{
+	case EColorSpace::None:
+		// leaves FVector2d::Zero from the constructor
+		break;
+	case EColorSpace::sRGB:
+		Chromaticities[0] = FVector2d(0.64, 0.33);
+		Chromaticities[1] = FVector2d(0.30, 0.60);
+		Chromaticities[2] = FVector2d(0.15, 0.06);
+		Chromaticities[3] = GetWhitePoint(EWhitePoint::CIE1931_D65);
+		break;
+	case EColorSpace::Rec2020:
+		Chromaticities[0] = FVector2d(0.708, 0.292);
+		Chromaticities[1] = FVector2d(0.170, 0.797);
+		Chromaticities[2] = FVector2d(0.131, 0.046);
+		Chromaticities[3] = GetWhitePoint(EWhitePoint::CIE1931_D65);
+		break;
+	case EColorSpace::ACESAP0:
+		Chromaticities[0] = FVector2d(0.7347, 0.2653);
+		Chromaticities[1] = FVector2d(0.0000, 1.0000);
+		Chromaticities[2] = FVector2d(0.0001, -0.0770);
+		Chromaticities[3] = GetWhitePoint(EWhitePoint::ACES_D60);
+		break;
+	case EColorSpace::ACESAP1:
+		Chromaticities[0] = FVector2d(0.713, 0.293);
+		Chromaticities[1] = FVector2d(0.165, 0.830);
+		Chromaticities[2] = FVector2d(0.128, 0.044);
+		Chromaticities[3] = GetWhitePoint(EWhitePoint::ACES_D60);
+		break;
+	case EColorSpace::P3DCI:
+		Chromaticities[0] = FVector2d(0.680, 0.320);
+		Chromaticities[1] = FVector2d(0.265, 0.690);
+		Chromaticities[2] = FVector2d(0.150, 0.060);
+		Chromaticities[3] = GetWhitePoint(EWhitePoint::DCI_CalibrationWhite);
+		break;
+	case EColorSpace::P3D65:
+		Chromaticities[0] = FVector2d(0.680, 0.320);
+		Chromaticities[1] = FVector2d(0.265, 0.690);
+		Chromaticities[2] = FVector2d(0.150, 0.060);
+		Chromaticities[3] = GetWhitePoint(EWhitePoint::CIE1931_D65);
+		break;
+	case EColorSpace::REDWideGamut:
+		Chromaticities[0] = FVector2d(0.780308, 0.304253);
+		Chromaticities[1] = FVector2d(0.121595, 1.493994);
+		Chromaticities[2] = FVector2d(0.095612, -0.084589);
+		Chromaticities[3] = GetWhitePoint(EWhitePoint::CIE1931_D65);
+		break;
+	case EColorSpace::SonySGamut3:
+		Chromaticities[0] = FVector2d(0.730, 0.280);
+		Chromaticities[1] = FVector2d(0.140, 0.855);
+		Chromaticities[2] = FVector2d(0.100, -0.050);
+		Chromaticities[3] = GetWhitePoint(EWhitePoint::CIE1931_D65);
+		break;
+	case EColorSpace::SonySGamut3Cine:
+		Chromaticities[0] = FVector2d(0.766, 0.275);
+		Chromaticities[1] = FVector2d(0.225, 0.800);
+		Chromaticities[2] = FVector2d(0.089, -0.087);
+		Chromaticities[3] = GetWhitePoint(EWhitePoint::CIE1931_D65);
+		break;
+	case EColorSpace::AlexaWideGamut:
+		Chromaticities[0] = FVector2d(0.684, 0.313);
+		Chromaticities[1] = FVector2d(0.221, 0.848);
+		Chromaticities[2] = FVector2d(0.0861, -0.1020);
+		Chromaticities[3] = GetWhitePoint(EWhitePoint::CIE1931_D65);
+		break;
+	case EColorSpace::CanonCinemaGamut:
+		Chromaticities[0] = FVector2d(0.740, 0.270);
+		Chromaticities[1] = FVector2d(0.170, 1.140);
+		Chromaticities[2] = FVector2d(0.080, -0.100);
+		Chromaticities[3] = GetWhitePoint(EWhitePoint::CIE1931_D65);
+		break;
+	case EColorSpace::GoProProtuneNative:
+		Chromaticities[0] = FVector2d(0.698448, 0.193026);
+		Chromaticities[1] = FVector2d(0.329555, 1.024597);
+		Chromaticities[2] = FVector2d(0.108443, -0.034679);
+		Chromaticities[3] = GetWhitePoint(EWhitePoint::CIE1931_D65);
+		break;
+	case EColorSpace::PanasonicVGamut:
+		Chromaticities[0] = FVector2d(0.730, 0.280);
+		Chromaticities[1] = FVector2d(0.165, 0.840);
+		Chromaticities[2] = FVector2d(0.100, -0.030);
+		Chromaticities[3] = GetWhitePoint(EWhitePoint::CIE1931_D65);
+		break;
+	case EColorSpace::PLASA_E1_54:
+		Chromaticities[0] = FVector2d(0.7347, 0.2653);
+		Chromaticities[1] = FVector2d(0.1596, 0.8404);
+		Chromaticities[2] = FVector2d(0.0366, 0.0001);
+		Chromaticities[3] = FVector2d(0.4254, 0.4044);
+		break;
+	default:
+		checkNoEntry();
+		break;
+	}
+
+	return Chromaticities;
+}
+
+FMatrix44d FColorSpace::CalcRgbToXYZ() const
+{
+	FMatrix44d Mat = FMatrix44d(
+		FVector3d(Chromaticities[0].X, Chromaticities[0].Y, 1.0 - Chromaticities[0].X - Chromaticities[0].Y),
+		FVector3d(Chromaticities[1].X, Chromaticities[1].Y, 1.0 - Chromaticities[1].X - Chromaticities[1].Y),
+		FVector3d(Chromaticities[2].X, Chromaticities[2].Y, 1.0 - Chromaticities[2].X - Chromaticities[2].Y),
+		FVector3d{0,0,0}
+	);
+
+	FMatrix44d Inverse = Mat.Inverse();
+	const FVector2d& White = Chromaticities[3];
+	FVector3d WhiteXYZ = FVector3d(White[0] / White[1], 1.0, (1.0f - White[0] - White[1]) / White[1]);
+	FVector3d Scale = Inverse.TransformVector(WhiteXYZ);
+
+	for (unsigned i = 0; i < 3; i++)
+	{
+		Mat.M[0][i] *= Scale[0];
+		Mat.M[1][i] *= Scale[1];
+		Mat.M[2][i] *= Scale[2];
+	}
+	return Mat;
+}
+
+bool FColorSpace::Equals(const FColorSpace& ColorSpace, double Tolerance) const
+{
+	return Equals(ColorSpace.Chromaticities, Tolerance);
+}
+
+bool FColorSpace::Equals(const TStaticArray<FVector2d, 4>& InChromaticities, double Tolerance) const
+{
+	return Chromaticities[0].Equals(InChromaticities[0], Tolerance) &&
+		Chromaticities[1].Equals(InChromaticities[1], Tolerance) &&
+		Chromaticities[2].Equals(InChromaticities[2], Tolerance) &&
+		Chromaticities[3].Equals(InChromaticities[3], Tolerance);
+}
+
+bool FColorSpace::IsSRGB() const
+{
+	return bIsSRGB;
+}
+
+FLinearColor FColorSpace::MakeFromColorTemperature(float Temp) const
+{
+	Temp = FMath::Clamp(Temp, 1000.0f, 15000.0f);
+
+	// Approximate Planckian locus in CIE 1960 UCS
+	float u = (0.860117757f + 1.54118254e-4f * Temp + 1.28641212e-7f * Temp * Temp) / (1.0f + 8.42420235e-4f * Temp + 7.08145163e-7f * Temp * Temp);
+	float v = (0.317398726f + 4.22806245e-5f * Temp + 4.20481691e-8f * Temp * Temp) / (1.0f - 2.89741816e-5f * Temp + 1.61456053e-7f * Temp * Temp);
+
+	float x = 3.0f * u / (2.0f * u - 8.0f * v + 4.0f);
+	float y = 2.0f * v / (2.0f * u - 8.0f * v + 4.0f);
+	float z = 1.0f - x - y;
+
+	FVector3d XYZ = FVector3d(1.0 / y * x, 1.0, 1.0 / y * z);
+	FVector4d RGB = XYZToRgb.TransformVector(XYZ);
+
+	// The XYZ to RGB transform can result in negative values, so we need to clamp here.
+	return FLinearColor(FMath::Max(0.0f, (float)RGB.X), FMath::Max(0.0f, (float)RGB.Y), FMath::Max(0.0f, (float)RGB.Z));
+}
+
+float FColorSpace::GetLuminance(const FLinearColor& Color) const
+{
+	//Note: Equivalent to the dot product of Color and RgbToXYZ.GetColumn(1).
+	return Color.R * (float)RgbToXYZ.M[0][1] + Color.G * (float)RgbToXYZ.M[1][1] + Color.B * (float)RgbToXYZ.M[2][1];
+}
+
+FLinearColor FColorSpace::GetLuminanceFactors() const
+{
+	return FLinearColor(RgbToXYZ.GetColumn(1));
+}
+
+
+FMatrix44d FColorSpaceTransform::CalcChromaticAdaptionMatrix(FVector3d SourceXYZ, FVector3d TargetXYZ, EChromaticAdaptationMethod Method)
+{
+	FMatrix44d XyzToRgb;
+
+	if (Method == EChromaticAdaptationMethod::CAT02)
+	{
+		XyzToRgb = FMatrix44d(
+			FPlane4d(0.7328,  0.4296, -0.1624,  0.),
+			FPlane4d(-0.7036,  1.6975,  0.0061,  0.),
+			FPlane4d(0.0030,  0.0136,  0.9834,  0.),
+			FPlane4d(0., 0., 0., 1.)
+		).GetTransposed();
+	}
+	else if (Method == EChromaticAdaptationMethod::Bradford)
+	{
+		XyzToRgb = FMatrix44d(
+			FPlane4d(0.8951,  0.2664, -0.1614,  0.),
+			FPlane4d(-0.7502,  1.7135,  0.0367,  0.),
+			FPlane4d(0.0389, -0.0685,  1.0296,  0.),
+			FPlane4d(0., 0., 0., 1.)
+		).GetTransposed();
+	}
+	else
+	{
+		return FMatrix44d::Identity;
+	}
+
+	FVector4d SourceRGB = XyzToRgb.TransformVector(SourceXYZ);
+	FVector4d TargetRGB = XyzToRgb.TransformVector(TargetXYZ);
+	
+	FVector4d Scale = FVector4d(
+		TargetRGB.X / SourceRGB.X,
+		TargetRGB.Y / SourceRGB.Y,
+		TargetRGB.Z / SourceRGB.Z,
+		1.0);
+
+	FMatrix44d ScaleMat = FMatrix44d::Identity;
+	ScaleMat.M[0][0] = Scale.X;
+	ScaleMat.M[1][1] = Scale.Y;
+	ScaleMat.M[2][2] = Scale.Z;
+	ScaleMat.M[3][3] = Scale.W;
+
+	FMatrix44d RgbToXyz = XyzToRgb.Inverse();
+
+	return XyzToRgb * ScaleMat * RgbToXyz;
+}
+
+FColorSpaceTransform FColorSpaceTransform::GetSRGBToWorkingColorSpace()
+{
+	static FColorSpaceTransform CachedTransform = FColorSpaceTransform(FColorSpace(EColorSpace::sRGB), FColorSpace::GetWorking());
+	
+	return CachedTransform;
+}
+
+static FMatrix44d CalcColorSpaceTransformMatrix(const FColorSpace& Src, const FColorSpace& Dst, EChromaticAdaptationMethod Method)
+{
+	if (Method == UE::Color::EChromaticAdaptationMethod::None)
+	{
+		return Src.GetRgbToXYZ() * Dst.GetXYZToRgb();
+	}
+
+	const FVector3d SrcWhite = Src.GetRgbToXYZ().TransformVector(FVector3d::One());
+	const FVector3d DstWhite = Dst.GetRgbToXYZ().TransformVector(FVector3d::One());
+
+	if (SrcWhite.Equals(DstWhite, 1.e-7))
+	{
+		return Src.GetRgbToXYZ() * Dst.GetXYZToRgb();
+	}
+
+	FMatrix44d ChromaticAdaptationMat = FColorSpaceTransform::CalcChromaticAdaptionMatrix(SrcWhite, DstWhite, Method);
+	return Src.GetRgbToXYZ() * ChromaticAdaptationMat * Dst.GetXYZToRgb();
+}
+
+FColorSpaceTransform::FColorSpaceTransform(const FColorSpace& Src, const FColorSpace& Dst, EChromaticAdaptationMethod Method)
+	: FMatrix44d(CalcColorSpaceTransformMatrix(Src,Dst,Method))
+{
+
+}
+
+FColorSpaceTransform::FColorSpaceTransform(FMatrix44d Matrix)
+	: FMatrix44d(MoveTemp(Matrix))
+{ }
+
+FLinearColor FColorSpaceTransform::Apply(const FLinearColor& Color) const
+{
+	FLinearColor Result;
+	VectorRegister4Float VecP = VectorLoad(&Color.R);
+	VectorRegister4Float VecR = VectorTransformVector(VecP, this);
+	VectorStore(VecR, &Result.R);
+	return Result;
+}
+
+
+} } // end namespace UE::Color
+
diff --git a/Engine/Source/Runtime/Core/Private/ColorManagement/Tests/ColorManagementTests.cpp b/Engine/Source/Runtime/Core/Private/ColorManagement/Tests/ColorManagementTests.cpp
new file mode 100644
--- /dev/null
+++ b/Engine/Source/Runtime/Core/Private/ColorManagement/Tests/ColorManagementTests.cpp
@@ -0,0 +1,306 @@
+// Copyright Epic Games, Inc. All Rights Reserved.
+
+#if WITH_TESTS
+
+#include "CoreTypes.h"
+#include "Containers/UnrealString.h"
+#include "ColorManagement/TransferFunctions.h"
+#include "ColorManagement/ColorSpace.h"
+#include "Math/NumericLimits.h"
+#include "Math/UnrealMathUtility.h"
+#include "Logging/LogMacros.h"
+#include "Tests/TestHarnessAdapter.h"
+
+DEFINE_LOG_CATEGORY_STATIC(LogUnrealColorManagementTest, Log, All);
+
+TEST_CASE_NAMED(FTransferFunctionsTest, "System::ColorManagement::TransferFunctions", "[EditorContext][EngineFilter]")
+{
+	using namespace UE::Color;
+
+	const float TestIncrement = 0.05f;
+
+	// Verify that all transfer functions correctly inverse each other.
+
+	for (uint8 EnumValue = static_cast<uint8>(EEncoding::Linear); EnumValue < static_cast<uint8>(EEncoding::Max); EnumValue++)
+	{
+		EEncoding EncodingType = static_cast<EEncoding>(EnumValue);
+
+		for (float TestValue = 0.0f; TestValue <= 1.0f; TestValue += TestIncrement)
+		{
+			float Encoded = UE::Color::Encode(EncodingType, TestValue);
+			float Decoded = UE::Color::Decode(EncodingType, Encoded);
+			CHECK_MESSAGE(TEXT("Transfer function encode followed by decode must match identity"), FMath::IsNearlyEqual(Decoded, TestValue, UE_KINDA_SMALL_NUMBER));
+		}
+	}
+}
+
+/**
+ * Tests if two double matrices (4x4 xyzw) are equal within an optional tolerance
+ *
+ * @param Mat0 First Matrix
+ * @param Mat1 Second Matrix
+ * @param Tolerance Error per item allowed for the comparison
+ *
+ * @return true if equal within tolerance
+ */
+static bool TestMatricesDoubleEqual(FMatrix44d& Mat0, FMatrix44d& Mat1, double Tolerance = 0.0)
+{
+	for (int32 Row = 0; Row < 4; ++Row)
+	{
+		for (int32 Column = 0; Column < 4; ++Column)
+		{
+			double Diff = Mat0.M[Row][Column] - Mat1.M[Row][Column];
+			if (FMath::Abs(Diff) > Tolerance)
+			{
+				UE_LOG(LogUnrealColorManagementTest, Log, TEXT("Bad(%.8f) at [%i, %i]"), Diff, Row, Column);
+				return false;
+			}
+		}
+	}
+	return true;
+}
+
+static void TestRgbToXyzConversionMatrices()
+{
+	// Note: test matrices generated using the python (colour-science) colour library.
+
+	using namespace UE::Color;
+
+	double Tolerance = 0.000001;
+
+	FMatrix44d Mat0 = FColorSpace(EColorSpace::sRGB).GetRgbToXYZ();
+	FMatrix44d Mat1 = FMatrix44d(
+		{0.412390799266, 0.357584339384, 0.180480788402, 0.0},
+		{0.212639005872, 0.715168678768, 0.0721923153607, 0.0},
+		{0.0193308187156, 0.119194779795, 0.95053215225, 0.0},
+		{0,0,0,1}
+	).GetTransposed();
+	CHECK_MESSAGE(TEXT("Rec709 RGB2XYZ Matrix Equality"), TestMatricesDoubleEqual(Mat0, Mat1, Tolerance));
+
+	Mat0 = FColorSpace(EColorSpace::Rec2020).GetRgbToXYZ();
+	Mat1 = FMatrix44d(
+		{0.636958048301, 0.144616903586, 0.168880975164, 0.0},
+		{0.262700212011, 0.677998071519, 0.0593017164699, 0.0},
+		{4.99410657447e-17, 0.0280726930491, 1.06098505771, 0.0},
+		{0,0,0,1}
+	).GetTransposed();
+	CHECK_MESSAGE(TEXT("Rec2020 RGB2XYZ Matrix Equality"), TestMatricesDoubleEqual(Mat0, Mat1, Tolerance));
+
+	Mat0 = FColorSpace(EColorSpace::ACESAP0).GetRgbToXYZ();
+	Mat1 = FMatrix44d(
+		{ 0.9525523959, 0.0, 9.36786e-05, 0.0 },
+		{ 0.3439664498, 0.7281660966, -0.0721325464, 0.0 },
+		{ 0.0, 0.0, 1.0088251844, 0.0 },
+		{ 0,0,0,1 }
+	).GetTransposed();
+	CHECK_MESSAGE(TEXT("AP0 RGB2XYZ Matrix Equality"), TestMatricesDoubleEqual(Mat0, Mat1, Tolerance));
+
+	Mat0 = FColorSpace(EColorSpace::ACESAP1).GetRgbToXYZ();
+	Mat1 = FMatrix44d(
+		{ 0.662454181109, 0.134004206456, 0.156187687005, 0.0 },
+		{ 0.272228716781, 0.674081765811, 0.0536895174079, 0.0 },
+		{ -0.00557464949039, 0.00406073352898, 1.01033910031, 0.0 },
+		{ 0,0,0,1 }
+	).GetTransposed();
+	CHECK_MESSAGE(TEXT("AP1 RGB2XYZ Matrix Equality"), TestMatricesDoubleEqual(Mat0, Mat1, Tolerance));
+
+	Mat0 = FColorSpace(EColorSpace::P3DCI).GetRgbToXYZ();
+	Mat1 = FMatrix44d(
+		{0.445169815565, 0.277134409207, 0.172282669816, 0.0},
+		{0.209491677913, 0.721595254161, 0.0689130679262, 0.0},
+		{-3.63410131697e-17, 0.047060560054, 0.907355394362, 0.0},
+		{0,0,0,1}
+	).GetTransposed();
+	CHECK_MESSAGE(TEXT("P3DCI RGB2XYZ Matrix Equality"), TestMatricesDoubleEqual(Mat0, Mat1, Tolerance));
+
+	Mat0 = FColorSpace(EColorSpace::P3D65).GetRgbToXYZ();
+	Mat1 = FMatrix44d(
+		{0.486570948648, 0.265667693169, 0.198217285234, 0.0},
+		{0.22897456407, 0.691738521837, 0.0792869140937, 0.0},
+		{-3.97207551693e-17, 0.0451133818589, 1.0439443689, 0.0},
+		{0,0,0,1}
+	).GetTransposed();
+	CHECK_MESSAGE(TEXT("P3D65 RGB2XYZ Matrix Equality"), TestMatricesDoubleEqual(Mat0, Mat1, Tolerance));
+}
+
+
+static void TestColorSpaceTransforms(UE::Color::EChromaticAdaptationMethod Method)
+{
+	// Note: test matrices generated using the python (colour-science) colour library.
+
+	using namespace UE::Color;
+
+	double Tolerance = UE_SMALL_NUMBER;
+
+	const FColorSpace Src = FColorSpace(EColorSpace::ACESAP1);
+	FMatrix44d Mat0, Mat1;
+
+	if (Method == UE::Color::EChromaticAdaptationMethod::None)
+	{
+		Mat0 = FColorSpaceTransform(Src, FColorSpace(EColorSpace::sRGB), Method);
+		Mat1 = FMatrix44d(
+			{ 1.73125381945, -0.604043087283, -0.0801077089571, 0.0 },
+			{ -0.131618928589, 1.13484150569, -0.00867943255179, 0.0 },
+			{ -0.0245682525938, -0.125750404281, 1.06563695775, 0.0 },
+			{0,0,0,1}
+		).GetTransposed();
+		CHECK_MESSAGE(TEXT("AP1->Rec709 without chromatic adaptation"), TestMatricesDoubleEqual(Mat0, Mat1, Tolerance));
+
+		Mat0 = FColorSpaceTransform(Src, FColorSpace(EColorSpace::Rec2020), Method);
+		Mat1 = FMatrix44d(
+			{ 1.04179138412, -0.0107415627227, -0.00696187506631, 0.0 },
+			{ -0.00168312771738, 1.00036605073, -0.0014082109903, 0.0 },
+			{ -0.005209686529, -0.0226414456785, 0.95230241486, 0.0 },
+			{ 0,0,0,1 }
+		).GetTransposed();
+		CHECK_MESSAGE(TEXT("AP1->Rec709 without chromatic adaptation"), TestMatricesDoubleEqual(Mat0, Mat1, Tolerance));
+
+		Mat0 = FColorSpaceTransform(Src, FColorSpace(EColorSpace::ACESAP0), Method);
+		Mat1 = FMatrix44d(
+			{ 0.695452241359, 0.140678696471, 0.163869062214, 0.0 },
+			{ 0.0447945633525, 0.859671118443, 0.0955343182103, 0.0 },
+			{ -0.00552588255811, 0.00402521030598, 1.00150067225, 0.0 },
+			{ 0,0,0,1 }
+		).GetTransposed();
+		CHECK_MESSAGE(TEXT("AP1->Rec709 without chromatic adaptation"), TestMatricesDoubleEqual(Mat0, Mat1, Tolerance));
+
+		Mat0 = FColorSpaceTransform(Src, FColorSpace(EColorSpace::ACESAP1), Method);
+		Mat1 = FMatrix44d::Identity;
+		CHECK_MESSAGE(TEXT("AP1->Rec709 without chromatic adaptation"), TestMatricesDoubleEqual(Mat0, Mat1, Tolerance));
+
+		Mat0 = FColorSpaceTransform(Src, FColorSpace(EColorSpace::P3DCI), Method);
+		Mat1 = FMatrix44d(
+			{ 1.53077277413, -0.322790385146, -0.0736971869439, 0.0 },
+			{ -0.0668950445367, 1.03255367118, -0.0105932139741, 0.0 },
+			{ -0.0026742897488, -0.0490786970568, 1.11404817691, 0.0 },
+			{ 0,0,0,1 }
+		).GetTransposed();
+		CHECK_MESSAGE(TEXT("AP1->Rec709 without chromatic adaptation"), TestMatricesDoubleEqual(Mat0, Mat1, Tolerance));
+
+		Mat0 = FColorSpaceTransform(Src, FColorSpace(EColorSpace::P3D65), Method);
+		Mat1 = FMatrix44d(
+			{ 1.40052305923, -0.295324940013, -0.067426473386, 0.0 },
+			{ -0.069782360156, 1.07712062473, -0.0110504369624, 0.0 },
+			{ -0.0023243874884, -0.0426572735573, 0.968286867584, 0.0 },
+			{ 0,0,0,1 }
+		).GetTransposed();
+		CHECK_MESSAGE(TEXT("AP1->Rec709 without chromatic adaptation"), TestMatricesDoubleEqual(Mat0, Mat1, Tolerance));
+	}
+	else if (Method == UE::Color::EChromaticAdaptationMethod::Bradford)
+	{
+		Mat0 = FColorSpaceTransform(Src, FColorSpace(EColorSpace::sRGB), Method);
+		Mat1 = FMatrix44d(
+			{ 1.70505099266, -0.621792120657, -0.083258872001, 0.0 },
+			{ -0.130256417507, 1.14080473658, -0.0105483190684, 0.0 },
+			{ -0.0240033568046, -0.128968976065, 1.15297233287, 0.0 },
+			{ 0,0,0,1 }
+		).GetTransposed();
+		CHECK_MESSAGE(TEXT("AP1->Rec709 with Bradford chromatic adaptation"), TestMatricesDoubleEqual(Mat0, Mat1, Tolerance));
+
+		Mat0 = FColorSpaceTransform(Src, FColorSpace(EColorSpace::Rec2020), Method);
+		Mat1 = FMatrix44d(
+			{ 1.02582474767, -0.0200531908382, -0.0057715568278, 0.0 },
+			{ -0.00223436951998, 1.00458650189, -0.0023521323685, 0.0 },
+			{ -0.00501335146809, -0.0252900718108, 1.03030342328, 0.0 },
+			{ 0,0,0,1 }
+		).GetTransposed();
+		CHECK_MESSAGE(TEXT("AP1->Rec709 with Bradford chromatic adaptation"), TestMatricesDoubleEqual(Mat0, Mat1, Tolerance));
+
+		Mat0 = FColorSpaceTransform(Src, FColorSpace(EColorSpace::ACESAP0), Method);
+		Mat1 = FMatrix44d(
+			{ 0.695452241359, 0.140678696471, 0.163869062214, 0.0 },
+			{ 0.0447945633525, 0.859671118443, 0.0955343182103, 0.0 },
+			{ -0.00552588255811, 0.00402521030598, 1.00150067225, 0.0 },
+			{ 0,0,0,1 }
+		).GetTransposed();
+		CHECK_MESSAGE(TEXT("AP1->Rec709 with Bradford chromatic adaptation"), TestMatricesDoubleEqual(Mat0, Mat1, Tolerance));
+
+		Mat0 = FColorSpaceTransform(Src, FColorSpace(EColorSpace::ACESAP1), Method);
+		Mat1 = FMatrix44d::Identity;
+		CHECK_MESSAGE(TEXT("AP1->Rec709 with Bradford chromatic adaptation"), TestMatricesDoubleEqual(Mat0, Mat1, Tolerance));
+
+		Mat0 = FColorSpaceTransform(Src, FColorSpace(EColorSpace::P3DCI), Method);
+		Mat1 = FMatrix44d(
+			{ 1.46412016696, -0.393327041647, -0.0707931253151, 0.0 },
+			{ -0.0664765138416, 1.07529152526, -0.00881501141699, 0.0 },
+			{ -0.00255286167095, -0.0470296027287, 1.0495824644, 0.0 },
+			{ 0,0,0,1 }
+		).GetTransposed();
+		CHECK_MESSAGE(TEXT("AP1->Rec709 with Bradford chromatic adaptation"), TestMatricesDoubleEqual(Mat0, Mat1, Tolerance));
+
+		Mat0 = FColorSpaceTransform(Src, FColorSpace(EColorSpace::P3D65), Method);
+		Mat1 = FMatrix44d(
+			{ 1.37921412825, -0.308864144674, -0.0703499835796, 0.0 },
+			{ -0.0693348583814, 1.082296746, -0.012961887621, 0.0 },
+			{ -0.00215900951357, -0.0454593248373, 1.04761833435, 0.0 },
+			{ 0,0,0,1 }
+		).GetTransposed();
+		CHECK_MESSAGE(TEXT("AP1->Rec709 with Bradford chromatic adaptation"), TestMatricesDoubleEqual(Mat0, Mat1, Tolerance));
+	}
+}
+
+static void TestAppliedColorTransforms()
+{
+	using namespace UE::Color;
+
+	// Intentionally misalign the start of FLinearColor member,
+	// as a test against 16-byte (128bit) alignment when using 4 float SIMD instructions.
+	// 
+	// Given that we now load/store from unaligned memory, the above alignment requirement is lifted.
+	struct FAlignmentTest
+	{
+		float Nudge = 0.0f;
+		FLinearColor SrcColor = FLinearColor(1.0f, 0.5f, 0.0f);
+	};
+
+	FColorSpaceTransform Transform = FColorSpaceTransform(FColorSpace(EColorSpace::sRGB), FColorSpace(EColorSpace::ACESAP1), EChromaticAdaptationMethod::Bradford);
+
+	TUniquePtr<FAlignmentTest> AlignmentTest = MakeUnique<FAlignmentTest>();
+	FLinearColor ExpectedResult = FLinearColor(0.78285898f, 0.52837066f, 0.07540048f);
+	FLinearColor Result = Transform.Apply(AlignmentTest->SrcColor);
+
+	CHECK_MESSAGE(TEXT("FLinearColor sRGB->AP1 color transform"), Result.Equals(ExpectedResult, UE_SMALL_NUMBER));
+}
+
+static void TestLuminance()
+{
+	using namespace UE::Color;
+
+	// Note: test factors from the python (colour-science) colour library.
+
+	FLinearColor LuminanceFactors = FColorSpace(EColorSpace::sRGB).GetLuminanceFactors();
+	CHECK_MESSAGE(TEXT("sRGB luminance factors equality test"), LuminanceFactors.Equals(FLinearColor(0.212639005872f, 0.715168678768f, 0.0721923153607f), UE_SMALL_NUMBER));
+
+	LuminanceFactors = FColorSpace(EColorSpace::ACESAP1).GetLuminanceFactors();
+	CHECK_MESSAGE(TEXT("ACESAP1 luminance factors equality test"), LuminanceFactors.Equals(FLinearColor(0.272228716781f, 0.674081765811f, 0.0536895174079f), UE_SMALL_NUMBER));
+
+	LuminanceFactors = FColorSpace(EColorSpace::Rec2020).GetLuminanceFactors();
+	CHECK_MESSAGE(TEXT("Rec2020 luminance factors equality test"), LuminanceFactors.Equals(FLinearColor(0.262700212011f, 0.677998071519f, 0.0593017164699f), UE_SMALL_NUMBER));
+}
+
+TEST_CASE_NAMED(FColorSpaceTest, "System::ColorManagement::ColorSpace", "[EditorContext][EngineFilter]")
+{
+	using namespace UE::Color;
+
+	FColorSpace CS = FColorSpace(EColorSpace::sRGB);
+	FMatrix44d Mat0 = FColorSpaceTransform(CS, CS, EChromaticAdaptationMethod::None);
+	FMatrix44d Mat1 = FMatrix44d::Identity;
+	CHECK_MESSAGE(TEXT("Identity color space conversion to itself should match identity"), TestMatricesDoubleEqual(Mat0, Mat1, 0.00000001));
+
+	SECTION("RgbToXyzConversionMatrices")
+	TestRgbToXyzConversionMatrices();
+
+	SECTION("ColorSpaceTransforms")
+	TestColorSpaceTransforms(EChromaticAdaptationMethod::None);
+
+	SECTION("ColorSpaceTransforms_Bradford")
+	TestColorSpaceTransforms(EChromaticAdaptationMethod::Bradford);
+
+	SECTION("AppliedColorTransforms")
+	TestAppliedColorTransforms();
+
+	SECTION("Luminance")
+	TestLuminance();
+}
+
+#endif //WITH_TESTS
diff --git a/Engine/Source/Runtime/Core/Private/ColorManagement/TransferFunctions.cpp b/Engine/Source/Runtime/Core/Private/ColorManagement/TransferFunctions.cpp
new file mode 100644
--- /dev/null
+++ b/Engine/Source/Runtime/Core/Private/ColorManagement/TransferFunctions.cpp
@@ -0,0 +1,129 @@
+// Copyright Epic Games, Inc. All Rights Reserved.
+
+#include "ColorManagement/TransferFunctions.h"
+#include "Logging/LogMacros.h"
+
+DEFINE_LOG_CATEGORY_STATIC(LogTransferFunctions, Log, All);
+
+namespace UE { namespace Color {
+
+inline TFunction<float(float)> GetTransferFunction(EEncoding SourceEncoding, bool bIsEncode)
+{
+	switch (SourceEncoding)
+	{
+	case EEncoding::None:        return &Linear;
+	case EEncoding::Linear:      return &Linear;
+	case EEncoding::sRGB:        return bIsEncode ? &EncodeSRGB : &DecodeSRGB;
+	case EEncoding::ST2084:      return bIsEncode ? &EncodeST2084 : &DecodeST2084;
+	case EEncoding::Gamma22:     return bIsEncode ? &EncodeGamma22 : &DecodeGamma22;
+	case EEncoding::BT1886:      return bIsEncode ? &EncodeBT1886 : &DecodeBT1886;
+	case EEncoding::Gamma26:     return bIsEncode ? &EncodeGamma26 : &DecodeGamma26;
+	case EEncoding::Cineon:      return bIsEncode ? &EncodeCineon : &DecodeCineon;
+	case EEncoding::REDLog:      return bIsEncode ? &EncodeREDLog : &DecodeREDLog;
+	case EEncoding::REDLog3G10:  return bIsEncode ? &EncodeREDLog3G10 : &DecodeREDLog3G10;
+	case EEncoding::SLog1:       return bIsEncode ? &EncodeSLog1 : &DecodeSLog1;
+	case EEncoding::SLog2:       return bIsEncode ? &EncodeSLog2 : &DecodeSLog2;
+	case EEncoding::SLog3:       return bIsEncode ? &EncodeSLog3 : &DecodeSLog3;
+	case EEncoding::AlexaV3LogC: return bIsEncode ? &EncodeArriAlexaV3LogC : &DecodeArriAlexaV3LogC;
+	case EEncoding::CanonLog:    return bIsEncode ? &EncodeCanonLog : &DecodeCanonLog;
+	case EEncoding::ProTune:     return bIsEncode ? &EncodeGoProProTune : &DecodeGoProProTune;
+	case EEncoding::VLog:        return bIsEncode ? &EncodePanasonicVLog : &DecodePanasonicVLog;
+	default:
+		check(false);
+		break;
+	}
+
+	UE_LOG(LogTransferFunctions, Warning, TEXT("Failed to find valid transfer function for enum value %d."), int(SourceEncoding));
+
+	return nullptr;
+}
+
+TFunction<float(float)> GetEncodeFunction(EEncoding SourceEncoding)
+{
+	return GetTransferFunction(SourceEncoding, true);
+}
+
+TFunction<float(float)> GetDecodeFunction(EEncoding SourceEncoding)
+{
+	return GetTransferFunction(SourceEncoding, false);
+}
+
+float Encode(EEncoding SourceEncoding, float Value)
+{
+	if (TFunction<float(float)> TransferFn = GetEncodeFunction(SourceEncoding))
+	{
+		return TransferFn(Value);
+	}
+	else
+	{
+		return Value;
+	}
+}
+
+float Decode(EEncoding SourceEncoding, float Value)
+{
+	if (TFunction<float(float)> TransferFn = GetDecodeFunction(SourceEncoding))
+	{
+		return TransferFn(Value);
+	}
+	else
+	{
+		return Value;
+	}
+}
+
+TFunction<FLinearColor(const FLinearColor&)> GetColorEncodeFunction(EEncoding SourceEncoding)
+{
+	TFunction<float(float)> TransferFn = GetEncodeFunction(SourceEncoding);
+
+	if (!TransferFn)
+	{
+		return nullptr;
+	}
+
+	return [TransferFn](const FLinearColor& Color) -> FLinearColor
+	{
+		return FLinearColor(TransferFn(Color.R), TransferFn(Color.G), TransferFn(Color.B), Color.A);
+	};
+}
+
+TFunction<FLinearColor(const FLinearColor&)> GetColorDecodeFunction(EEncoding SourceEncoding)
+{
+	TFunction<float(float)> TransferFn = GetDecodeFunction(SourceEncoding);
+
+	if (!TransferFn)
+	{
+		return nullptr;
+	}
+
+	return [TransferFn](const FLinearColor& Color) -> FLinearColor
+	{
+		return FLinearColor(TransferFn(Color.R), TransferFn(Color.G), TransferFn(Color.B), Color.A);
+	};
+}
+
+FLinearColor Encode(EEncoding SourceEncoding, const FLinearColor& Color)
+{
+	if (TFunction<FLinearColor(const FLinearColor&)> TransferFn = GetColorEncodeFunction(SourceEncoding))
+	{
+		return TransferFn(Color);
+	}
+	else
+	{
+		return Color;
+	}
+}
+
+FLinearColor Decode(EEncoding SourceEncoding, const FLinearColor& Color)
+{
+	if (TFunction<FLinearColor(const FLinearColor&)> TransferFn = GetColorDecodeFunction(SourceEncoding))
+	{
+		return TransferFn(Color);
+	}
+	else
+	{
+		return Color;
+	}
+}
+
+} } // end namespace UE::Color
diff --git a/Engine/Source/Runtime/Core/Public/ColorManagement/ColorManagementDefines.h b/Engine/Source/Runtime/Core/Public/ColorManagement/ColorManagementDefines.h
new file mode 100644
--- /dev/null
+++ b/Engine/Source/Runtime/Core/Public/ColorManagement/ColorManagementDefines.h
@@ -0,0 +1,87 @@
+// Copyright Epic Games, Inc. All Rights Reserved.
+
+#pragma once
+
+#include "HAL/Platform.h"
+
+namespace UE { namespace Color {
+
+/** Increment upon breaking changes to the EEncoding enum. Note that changing this forces a rebuild of textures that rely on it. */
+constexpr uint32 ENCODING_TYPES_VER = 4;
+
+/** Increment upon breaking changes to the EColorSpace and EChromaticAdaptationMethod enums. Note that changing this forces a rebuild of textures that rely on it.*/
+constexpr uint32 COLORSPACE_VER = 2;
+
+/** List of available encodings/transfer functions.
+* 
+* NOTE: This list is replicated as a UENUM in TextureDefines.h, and both should always match.
+*/
+enum class EEncoding : uint8
+{
+	None = 0,
+	Linear = 1,
+	sRGB,
+	ST2084,
+	Gamma22,
+	BT1886,
+	Gamma26,
+	Cineon,
+	REDLog,
+	REDLog3G10,
+	SLog1,
+	SLog2,
+	SLog3,
+	AlexaV3LogC,
+	CanonLog,
+	ProTune,
+	VLog,
+	Max,
+};
+
+/** List of available color spaces. (Increment COLORSPACE_VER upon breaking changes to the list.)
+* 
+* NOTE: This list is partially replicated as a UENUM in TextureDefines.h: any type exposed to textures should match the enum value below.
+*/
+enum class EColorSpace : uint8
+{
+	None = 0,
+	sRGB = 1,
+	Rec2020 = 2,
+	ACESAP0 = 3,
+	ACESAP1 = 4,
+	P3DCI = 5,
+	P3D65 = 6,
+	REDWideGamut = 7,
+	SonySGamut3 = 8,
+	SonySGamut3Cine = 9,
+	AlexaWideGamut = 10,
+	CanonCinemaGamut = 11,
+	GoProProtuneNative = 12,
+	PanasonicVGamut = 13,
+	PLASA_E1_54 = 14,
+	Max,
+};
+
+
+/** List of available chromatic adaptation methods.
+* 
+* NOTE: This list is replicated as a UENUM in TextureDefines.h, and both should always match.
+*/
+enum class EChromaticAdaptationMethod : uint8
+{
+	None = 0,
+	Bradford = 1,
+	CAT02 = 2,
+	Max,
+};
+
+/** List of standard white points. */
+enum class EWhitePoint : uint8
+{
+	CIE1931_D65 = 0,
+	ACES_D60 = 1,
+	DCI_CalibrationWhite,
+	Max,
+};
+
+} } // end namespace UE::Color
diff --git a/Engine/Source/Runtime/Core/Public/ColorManagement/ColorSpace.h b/Engine/Source/Runtime/Core/Public/ColorManagement/ColorSpace.h
new file mode 100644
--- /dev/null
+++ b/Engine/Source/Runtime/Core/Public/ColorManagement/ColorSpace.h
@@ -0,0 +1,366 @@
+// Copyright Epic Games, Inc. All Rights Reserved.
+
+#pragma once
+
+#include "ColorManagementDefines.h"
+#include "Containers/StaticArray.h"
+#include "CoreTypes.h"
+#include "Math/Color.h"
+#include "Math/Matrix.h"
+#include "Math/UnrealMathSSE.h"
+#include "Math/Vector.h"
+#include "Misc/AssertionMacros.h"
+#include "Serialization/Archive.h"
+#include "Templates/UnrealTypeTraits.h"
+
+namespace UE { namespace Color {
+
+/**
+ * Convert chromaticity coordinate and luminance to CIE XYZ tristimulus values.
+ *
+ * @return FVector3d
+ */
+inline FVector3d xyYToXYZ(const FVector3d& xyY)
+{
+	double Divisor = FMath::Max(xyY[1], 1e-10);
+
+	return FVector3d(
+		xyY[0] * xyY[2] / Divisor,
+		xyY[2],
+		(1.0 - xyY[0] - xyY[1]) * xyY[2] / Divisor
+	);
+}
+
+/**
+ * Convert CIE XYZ tristimulus values to chromaticitiy coordinates and luminance.
+ *
+ * @return FVector3d
+ */
+inline FVector3d XYZToxyY(const FVector3d& XYZ)
+{
+	double Divisor = XYZ[0] + XYZ[1] + XYZ[2];
+	if (Divisor == 0.0)
+	{
+		Divisor = 1e-10;
+	}
+
+	return FVector3d(
+		XYZ[0] / Divisor,
+		XYZ[1] / Divisor,
+		XYZ[1]
+	);
+}
+
+/**
+ * Convenience function to get the transposed matrix, i.e. for pre-multiplied shader matrices.
+ *
+ * @return TMatrix<T>
+ */
+template<typename T>
+inline UE::Math::TMatrix<T> Transpose(const FMatrix44d& Transform)
+{
+	if constexpr (std::is_same_v<T, double>)
+	{
+		return Transform.GetTransposed();
+	}
+	else
+	{
+		return UE::Math::TMatrix<T>(Transform).GetTransposed();
+	}
+}
+
+/**
+ * Get standard white point coordinates.
+ *
+ * @param InWhitePoint White point type.
+ * @return FVector2d Chromaticity coordinates.
+ */
+inline FVector2d GetWhitePoint(EWhitePoint InWhitePoint)
+{
+	switch (InWhitePoint)
+	{
+	case UE::Color::EWhitePoint::CIE1931_D65:
+		return FVector2d(0.3127, 0.3290);
+	case UE::Color::EWhitePoint::ACES_D60:
+		return FVector2d(0.32168, 0.33767);
+	case UE::Color::EWhitePoint::DCI_CalibrationWhite:
+		return FVector2d(0.314, 0.351);
+
+	default:
+		checkNoEntry();
+		return FVector2d();
+	}
+}
+
+
+/** Color space definition as 4 chromaticity coordinates, in double precision internally. */
+class FColorSpace
+{
+public:
+	/**
+	 * Get the global engine working color space (as a singleton).
+	 * 
+	 * @return FColorSpace working color space
+	 */
+	static CORE_API const FColorSpace& GetWorking();
+
+	/**
+	 * Set the global engine working color space (as a singleton).
+	 * 
+	 * @param ColorSpace working color space
+	 */
+	static CORE_API void SetWorking(FColorSpace ColorSpace);
+
+
+	/** Constructor */
+	FColorSpace() {}
+
+	/**
+	* Constructor
+	*
+	* @param InRed Chromaticity 2D coordinates for the red color.
+	* @param InGreen Chromaticity 2D coordinates for the green color.
+	* @param InBlue Chromaticity 2D coordinates for the blue color.
+	* @param InWhite Chromaticity 2D coordinates for the white point.
+	*/
+	CORE_API explicit FColorSpace(const FVector2d& InRed, const FVector2d& InGreen, const FVector2d& InBlue, const FVector2d& InWhite);
+
+	/**
+	* Constructor
+	*
+	* @param ColorSpaceType Color space type.
+	*/
+	CORE_API explicit FColorSpace(UE::Color::EColorSpace ColorSpaceType);
+
+	FColorSpace(FColorSpace&&) = default;
+	FColorSpace(const FColorSpace&) = default;
+	FColorSpace& operator=(FColorSpace&&) = default;
+	FColorSpace& operator=(const FColorSpace&) = default;
+
+	/**
+	* Make the chromaticities of the color space type.
+	* 
+	* @return TStaticArray of four chromaticities
+	*/
+	CORE_API static TStaticArray<FVector2d, 4> MakeChromaticities(UE::Color::EColorSpace ColorSpaceType);
+
+	/**
+	* Getter for the color space chromaticity coordinates.
+	* 
+	* @param OutRed FVector2d for the red color chromaticity coordinate.
+	* @param OutGreen FVector2d for the green color chromaticity coordinate.
+	* @param OutBlue FVector2d for the blue color chromaticity coordinate.
+	* @param OutWhite FVector2d for the white color chromaticity coordinate.
+	*/
+	template<typename T>
+	inline void GetChromaticities(UE::Math::TVector2<T>& OutRed, UE::Math::TVector2<T>& OutGreen, UE::Math::TVector2<T>& OutBlue, UE::Math::TVector2<T>& OutWhite) const
+	{
+		OutRed		= Chromaticities[0];
+		OutGreen	= Chromaticities[1];
+		OutBlue		= Chromaticities[2];
+		OutWhite	= Chromaticities[3];
+	}
+
+	/**
+	* Gets the color space's red chromaticity coordinates.
+	*
+	* @return FVector2d xy coordinates.
+	*/
+	inline const FVector2d& GetRedChromaticity() const
+	{
+		return Chromaticities[0];
+	}
+
+	/**
+	* Gets the color space's green chromaticity coordinates.
+	*
+	* @return FVector2d xy coordinates.
+	*/
+	inline const FVector2d& GetGreenChromaticity() const
+	{
+		return Chromaticities[1];
+	}
+
+	/**
+	* Gets the color space's blue chromaticity coordinates.
+	*
+	* @return FVector2d xy coordinates.
+	*/
+	inline const FVector2d& GetBlueChromaticity() const
+	{
+		return Chromaticities[2];
+	}
+
+	/**
+	* Gets the color space's white point chromaticity coordinates.
+	*
+	* @return FVector2d xy coordinates.
+	*/
+	inline const FVector2d& GetWhiteChromaticity() const
+	{
+		return Chromaticities[3];
+	}
+
+	/**
+	* Gets the RGB-to-XYZ conversion matrix.
+	*
+	* @return FMatrix conversion matrix.
+	*/
+	inline const FMatrix44d& GetRgbToXYZ() const
+	{
+		return RgbToXYZ;
+	}
+
+	/**
+	* Gets the XYZ-to-RGB conversion matrix.
+	*
+	* @return FMatrix conversion matrix.
+	*/
+	inline const FMatrix44d& GetXYZToRgb() const
+	{
+		return XYZToRgb;
+	}
+
+	/**
+	 * Check against another color space for equality.
+	 *
+	 * @param ColorSpace The vector to check against.
+	 * @return true if the color spaces are equal, false otherwise.
+	 */
+	inline bool operator==(const FColorSpace& ColorSpace) const
+	{
+		return Chromaticities == ColorSpace.Chromaticities;
+	}
+
+	/**
+	 * Check against another color space for inequality.
+	 *
+	 * @param ColorSpace The color space to check against.
+	 * @return true if the color spaces are not equal, false otherwise.
+	 */
+	inline bool operator!=(const FColorSpace& ColorSpace) const
+	{
+		return Chromaticities != ColorSpace.Chromaticities;
+	}
+
+	/**
+	 * Check against another color space for equality, within specified error limits.
+	 *
+	 * @param ColorSpace The color space to check against.
+	 * @param Tolerance Error tolerance.
+	 * @return true if the color spaces are equal within tolerance limits, false otherwise.
+	 */
+	CORE_API bool Equals(const FColorSpace& ColorSpace, double Tolerance = 1.e-7) const;
+
+
+	/**
+	 * Check against color space chromaticities for equality, within specified error limits.
+	 *
+	 * @param InChromaticities The color space chromaticities to check against.
+	 * @param Tolerance Error tolerance.
+	 * @return true if the color spaces are equal within tolerance limits, false otherwise.
+	 */
+	CORE_API bool Equals(const TStaticArray<FVector2d, 4>& InChromaticities, double Tolerance = 1.e-7) const;
+
+	/**
+	 * Convenience function to verify if the color space matches the engine's default sRGB chromaticities.
+	 *
+	 * @return true if sRGB.
+	 */
+	CORE_API bool IsSRGB() const;
+
+	/**
+	* Converts temperature in Kelvins of a black body radiator to an RGB color in the current space.
+	*/
+	CORE_API FLinearColor MakeFromColorTemperature(float Temp) const;
+
+	/**
+	 * Calculate the color's luminance value in the current space.
+	 *
+	 * @param Color The sampled color.
+	 * @return float Luminance value.
+	 */
+	CORE_API float GetLuminance(const FLinearColor& Color) const;
+
+	/**
+	 * Get the luminance factors in the current space.
+	 *
+	 * @return FLinearColor Luminance factors.
+	 */
+	CORE_API FLinearColor GetLuminanceFactors() const;
+
+private:
+
+	CORE_API FMatrix44d CalcRgbToXYZ() const;
+
+	/** Red, green, blue, white chromaticities, in order. */
+	TStaticArray<FVector2d, 4> Chromaticities;
+
+	FMatrix44d RgbToXYZ;
+	FMatrix44d XYZToRgb;
+
+	bool bIsSRGB;
+
+public:
+
+	/**
+	 * Serializer.
+	 *
+	 * @param Ar The Serialization Archive.
+	 * @param CS The Color Space being serialized.
+	 * @return Reference to the Archive after serialization.
+	 */
+	friend FArchive& operator<<(FArchive& Ar, FColorSpace& CS)
+	{
+		return Ar << CS.Chromaticities;
+	}
+
+	bool Serialize(FArchive& Ar)
+	{
+		Ar << *this;
+		return true;
+	}
+};
+
+
+struct FColorSpaceTransform : FMatrix44d
+{
+	/**
+	* Constructor: create a color space transformation matrix from a source to a target color space using the RGB-XYZ-RGB conversions.
+	*
+	* @param Source Source color space.
+	* @param Target Target color space.
+	* @param Method Chromatic adapation method.
+	*/
+	CORE_API explicit FColorSpaceTransform(const FColorSpace& Src, const FColorSpace& Dst, EChromaticAdaptationMethod Method = EChromaticAdaptationMethod::Bradford);
+
+	/**
+	* Constructor: create a color space transformation matrix from a raw matrix.
+	*
+	* @param Matrix Color space transformation matrix.
+	*/
+	CORE_API explicit FColorSpaceTransform(FMatrix44d Matrix);
+
+	/**
+	* Apply color space transform to FLinearColor.
+	* 
+	* @param Color Color to transform.
+	*/
+	CORE_API FLinearColor Apply(const FLinearColor& Color) const;
+
+	/**
+	* Calculate the chromatic adaptation matrix using the specific method.
+	*
+	* @param SourceXYZ Source color in XYZ space.
+	* @param TargetXYZ Target color in XYZ space.
+	* @param Method Adaptation method (None, Bradford, CAT02).
+	*/
+	CORE_API static FMatrix44d CalcChromaticAdaptionMatrix(FVector3d SourceXYZ, FVector3d TargetXYZ, EChromaticAdaptationMethod Method = EChromaticAdaptationMethod::Bradford);
+
+	/**
+	* Convenience function to get a (statically cached) conversion from sRGB/Rec709 to the working color space.
+	*/
+	CORE_API static FColorSpaceTransform GetSRGBToWorkingColorSpace();
+};
+
+} }  // end namespace UE::Color
diff --git a/Engine/Source/Runtime/Core/Public/ColorManagement/TransferFunctions.h b/Engine/Source/Runtime/Core/Public/ColorManagement/TransferFunctions.h
new file mode 100644
--- /dev/null
+++ b/Engine/Source/Runtime/Core/Public/ColorManagement/TransferFunctions.h
@@ -0,0 +1,521 @@
+// Copyright Epic Games, Inc. All Rights Reserved.
+
+#pragma once
+
+#include "ColorManagementDefines.h"
+#include "CoreTypes.h"
+#include "GenericPlatform/GenericPlatformMath.h"
+#include "Math/Color.h"
+#include "Math/UnrealMathUtility.h"
+#include "Templates/Function.h"
+
+namespace UE { namespace Color {
+
+inline float Linear(float Value)
+{
+	return Value;
+}
+
+/**
+* Encode value to sRGB.
+*
+* @return float encoded value.
+*/
+inline float EncodeSRGB(float Value)
+{
+	if (Value <= 0.04045f / 12.92f)
+	{
+		return Value * 12.92f;
+	}
+	else
+	{
+		return  FGenericPlatformMath::Pow(Value, (1.0f / 2.4f)) * 1.055f - 0.055f;
+	}
+}
+
+/**
+* Decode value with an sRGB encoding.
+*
+* @return float decoded value.
+*/
+inline float DecodeSRGB(float Value)
+{
+	if (Value <= 0.04045f)
+	{
+		return Value / 12.92f;
+	}
+	else
+	{
+		return  FGenericPlatformMath::Pow((Value + 0.055f) / 1.055f, 2.4f);
+	}
+}
+
+/**
+* Encode value to SMPTE ST 2084:2014.
+*
+* @return float encoded value.
+*/
+inline float EncodeST2084(float Value)
+{
+	const float Lp = 10000.0f;
+	const float m1 = 2610 / 4096.0f * (1.0f / 4.0f);
+	const float m2 = 2523 / 4096.0f * 128.0f;
+	const float c1 = 3424 / 4096.0f;
+	const float c2 = 2413 / 4096.0f * 32.f;
+	const float c3 = 2392 / 4096.0f * 32.f;
+
+	Value = FGenericPlatformMath::Pow(Value / Lp, m1);
+	return FGenericPlatformMath::Pow((c1 + c2 * Value) / (c3 * Value + 1), m2);
+}
+
+/**
+* Decode value with a SMPTE ST 2084:2014 encoding.
+*
+* @return float decoded value.
+*/
+inline float DecodeST2084(float Value)
+{
+	const float Lp = 10000.0f;
+	const float m1 = 2610 / 4096.0f * (1.0f / 4.0f);
+	const float m2 = 2523 / 4096.0f * 128.0f;
+	const float c1 = 3424 / 4096.0f;
+	const float c2 = 2413 / 4096.0f * 32.f;
+	const float c3 = 2392 / 4096.0f * 32.f;
+
+	const float Vp = FGenericPlatformMath::Pow(Value, 1.0f / m2);
+	Value = FGenericPlatformMath::Max(0.0f, Vp - c1);
+	return FGenericPlatformMath::Pow((Value / (c2 - c3 * Vp)), 1.0f / m1) * Lp;
+}
+
+/**
+* Encode value to Gamma 2.2.
+*
+* @return float encoded value.
+*/
+inline float EncodeGamma22(float Value)
+{
+	return FGenericPlatformMath::Pow(Value, 1.0f / 2.2f);
+}
+
+
+/**
+* Decode value with a Gamma 2.2 encoding.
+*
+* @return float decoded value.
+*/
+inline float DecodeGamma22(float Value)
+{
+	return FGenericPlatformMath::Pow(Value, 2.2f);
+}
+
+/**
+* Encode value to ITU-R BT.1886.
+*
+* @return float encoded value.
+*/
+inline float EncodeBT1886(float Value)
+{
+	const float L_B = 0;
+	const float L_W = 1;
+	float Gamma = 2.40f;
+	float GammaInv = 1.0f / Gamma;
+	float N = FGenericPlatformMath::Pow(L_W, GammaInv) - FGenericPlatformMath::Pow(L_B, GammaInv);
+	float A = FGenericPlatformMath::Pow(N, Gamma);
+	float B = FGenericPlatformMath::Pow(L_B, GammaInv) / N;
+	return FGenericPlatformMath::Pow(Value / A, GammaInv) - B;
+}
+
+/**
+* Encode value to Gamma 2.6.
+*
+* @return float encoded value.
+*/
+inline float EncodeGamma26(float Value)
+{
+	return FGenericPlatformMath::Pow(Value, 1.0f / 2.6f);
+}
+
+
+/**
+* Decode value with a Gamma 2.6 encoding.
+*
+* @return float decoded value.
+*/
+inline float DecodeGamma26(float Value)
+{
+	return FGenericPlatformMath::Pow(Value, 2.6f);
+}
+
+/**
+* Decode value with an ITU-R BT.1886 encoding.
+*
+* @return float decoded value.
+*/
+inline float DecodeBT1886(float Value)
+{
+	const float L_B = 0;
+	const float L_W = 1;
+	float Gamma = 2.40f;
+	float GammaInv = 1.0f / Gamma;
+	float N = FGenericPlatformMath::Pow(L_W, GammaInv) - FGenericPlatformMath::Pow(L_B, GammaInv);
+	float A = FGenericPlatformMath::Pow(N, Gamma);
+	float B = FGenericPlatformMath::Pow(L_B, GammaInv) / N;
+	return A * FGenericPlatformMath::Pow(FGenericPlatformMath::Max(Value + B, 0.0f), Gamma);
+}
+
+/**
+* Encode value to Cineon.
+*
+* @return float encoded value.
+*/
+inline float EncodeCineon(float Value)
+{
+	const float BlackOffset = FGenericPlatformMath::Pow(10.0f, (95.0f - 685.0f) / 300.0f);
+	return (685.0f + 300.0f * FGenericPlatformMath::LogX(10.0f, Value * (1.0f - BlackOffset) + BlackOffset)) / 1023.0f;
+}
+
+/**
+* Decode value with a Cineon encoding.
+*
+* @return float decoded value.
+*/
+inline float DecodeCineon(float Value)
+{
+	const float BlackOffset = FGenericPlatformMath::Pow(10.0f, (95.0f - 685.0f) / 300.0f);
+	return (FGenericPlatformMath::Pow(10.0f, (1023.0f * Value - 685.0f) / 300.0f) - BlackOffset) / (1.0f - BlackOffset);
+}
+
+/**
+* Encode value to RED Log.
+*
+* @return float encoded value.
+*/
+inline float EncodeREDLog(float Value)
+{
+	const float BlackOffset = FGenericPlatformMath::Pow(10.0f, (0.0f - 1023.0f) / 511.0f);
+	return (1023.0f + 511.0f * FGenericPlatformMath::LogX(10.0f, Value * (1.0f - BlackOffset) + BlackOffset)) / 1023.0f;
+}
+
+/**
+* Decode value with a RED Log encoding.
+*
+* @return float decoded value.
+*/
+inline float DecodeREDLog(float Value)
+{
+	const float BlackOffset = FGenericPlatformMath::Pow(10.0f, (0.0f - 1023.0f) / 511.0f);
+	return (FGenericPlatformMath::Pow(10.0f, (1023.0f * Value - 1023.0f) / 511.0f) - BlackOffset) / (1.0f - BlackOffset);
+}
+
+/**
+* Encode value to RED Log3G10.
+*
+* @return float encoded value.
+*/
+inline float EncodeREDLog3G10(float Value)
+{
+	const float A = 0.224282f;
+	const float B = 155.975327f;
+	const float C = 0.01f;
+	const float G = 15.1927f;
+
+	Value += C;
+
+	if (Value < 0.0f)
+	{
+		return Value * G;
+	}
+	else
+	{
+		return FGenericPlatformMath::Sign(Value) * A * FGenericPlatformMath::LogX(10.0f, B * FGenericPlatformMath::Abs(Value) + 1.0f);
+	}
+}
+
+/**
+* Decode value with a RED Log3G10 encoding.
+*
+* @return float decoded value.
+*/
+inline float DecodeREDLog3G10(float Value)
+{
+	const float A = 0.224282f;
+	const float B = 155.975327f;
+	const float C = 0.01f;
+	const float G = 15.1927f;
+
+	if (Value < 0.0f)
+	{
+		Value /= G;
+	}
+	else
+	{
+		Value = FGenericPlatformMath::Sign(Value) * (FGenericPlatformMath::Pow(10.0f, FGenericPlatformMath::Abs(Value) / A ) - 1.0f) / B;
+	}
+
+	return Value - C;
+}
+
+/**
+* Encode value to Sony S-Log1.
+*
+* @return float encoded value.
+*/
+inline float EncodeSLog1(float Value)
+{
+	Value /= 0.9f;
+	Value = 0.432699f * FGenericPlatformMath::LogX(10.0f, Value + 0.037584f) + 0.616596f + 0.03f;
+	return (Value * 219.0f + 16.0f) * 4.0f / 1023.0f;
+}
+
+/**
+* Decode value with a Sony S-Log1 encoding.
+*
+* @return float decoded value.
+*/
+inline float DecodeSLog1(float Value)
+{
+	Value = ((Value * 1023.f) / 4.0f - 16.0f) / 219.0f;
+	Value = FGenericPlatformMath::Pow(10.0f, (Value - 0.616596f - 0.03f) / 0.432699f) - 0.037584f;
+	return Value * 0.9f;
+}
+
+/**
+* Encode value to Sony S-Log2.
+*
+* @return float encoded value.
+*/
+inline float EncodeSLog2(float Value)
+{
+	if (Value >= 0.0f)
+	{
+		return (64.0f + 876.0f * (0.432699f * FGenericPlatformMath::LogX(10.0f, 155.0f * Value / 197.1f + 0.037584f) + 0.646596f)) / 1023.f;
+	}
+	else
+	{
+		return (64.0f +876.0f * (Value * 3.53881278538813f / 0.9f) + 0.646596f + 0.030001222851889303f) / 1023.f;
+	}
+}
+
+/**
+* Decode value with a Sony S-Log2 encoding.
+*
+* @return float decoded value.
+*/
+inline float DecodeSLog2(float Value)
+{
+	if (Value >= (64.f + 0.030001222851889303f * 876.f) / 1023.f)
+	{
+		return 197.1f * (FGenericPlatformMath::Pow(10.0f, ((Value * 1023.f - 64.f) / 876.f - 0.646596f) / 0.432699f) - 0.037584f) / 155.f;
+	}
+	else
+	{
+		return 0.9f * ((Value * 1023.f - 64.f) / 876.f - 0.030001222851889303f) / 3.53881278538813f;
+	}
+}
+
+/**
+* Encode value to Sony S-Log3.
+*
+* @return float encoded value.
+*/
+inline float EncodeSLog3(float Value)
+{
+	if (Value >= 0.01125000f)
+	{
+		return (420.0f + FGenericPlatformMath::LogX(10.0f, (Value + 0.01f) / 0.19f) * 261.5f) / 1023.0f;
+	}
+	else
+	{
+		return (Value * 76.2102946929f / 0.01125f + 95.0f) / 1023.0f;
+	}
+}
+
+/**
+* Decode value with a Sony S-Log3 encoding.
+*
+* @return float decoded value.
+*/
+inline float DecodeSLog3(float Value)
+{
+	if (Value >= 171.2102946929f / 1023.0f)
+	{
+		return (FGenericPlatformMath::Pow(10.0f, (Value * 1023.0f - 420.f) / 261.5f)) * 0.19f - 0.01f;
+	}
+	else
+	{
+		return (Value * 1023.0f - 95.0f) * 0.01125000f / (171.2102946929f - 95.0f);
+	}
+}
+
+/**
+* Encode value to ARRI Alexa LogC.
+*
+* @return float encoded value.
+*/
+inline float EncodeArriAlexaV3LogC(float Value)
+{
+	const float cut = 0.010591f;
+	const float a = 5.555556f;
+	const float b = 0.052272f;
+	const float c = 0.247190f;
+	const float d = 0.385537f;
+	const float e = 5.367655f;
+	const float f = 0.092809f;
+
+	if (Value > cut)
+	{
+		return c * FGenericPlatformMath::LogX(10.0f, a * Value + b) + d;
+	}
+	else
+	{
+		return e * Value + f;
+	}
+}
+
+/**
+* Decode value with an ARRI Alexa LogC encoding.
+*
+* @return float decoded value.
+*/
+inline float DecodeArriAlexaV3LogC(float Value)
+{
+	const float cut = 0.010591f;
+	const float a = 5.555556f;
+	const float b = 0.052272f;
+	const float c = 0.247190f;
+	const float d = 0.385537f;
+	const float e = 5.367655f;
+	const float f = 0.092809f;
+	
+	if (Value > e * cut + f)
+	{
+		return (FGenericPlatformMath::Pow(10.0f, (Value - d) / c) - b) / a;
+	}
+	else
+	{
+		return (Value - f) / e;
+	}
+}
+
+/**
+* Encode value to Canon Log.
+*
+* @return float encoded value.
+*/
+inline float EncodeCanonLog(float Value)
+{
+	if (Value < 0.0f)
+	{
+		return -(0.529136f * (FGenericPlatformMath::LogX(10.0f, -Value * 10.1596f + 1.0f)) - 0.0730597f);
+	}
+	else
+	{
+		return 0.529136f * FGenericPlatformMath::LogX(10.0f, 10.1596f * Value + 1.0f) + 0.0730597f;
+	}
+}
+
+
+/**
+* Decode value with a Canon Log encoding.
+*
+* @return float decoded value.
+*/
+inline float DecodeCanonLog(float Value)
+{
+	if (Value < 0.0730597f)
+	{
+		return -(FGenericPlatformMath::Pow(10.0f, (0.0730597f - Value) / 0.529136f) - 1.0f) / 10.1596f;
+	}
+	else
+	{
+		return (FGenericPlatformMath::Pow(10.0f, (Value - 0.0730597f) / 0.529136f) - 1.0f) / 10.1596f;
+	}
+}
+
+/**
+* Encode value to GoPro ProTune.
+*
+* @return float encoded value.
+*/
+inline float EncodeGoProProTune(float Value)
+{
+	return FGenericPlatformMath::Loge(Value * 112.f + 1.0f) / FGenericPlatformMath::Loge(113.0f);
+}
+
+/**
+* Decode value with a GoPro ProTune encoding.
+*
+* @return float decoded value.
+*/
+inline float DecodeGoProProTune(float Value)
+{
+	return (FGenericPlatformMath::Pow(113.f, Value) - 1.0f) / 112.f;
+}
+
+/**
+* Encode value to Panasonic V-Log.
+*
+* @return float encoded value.
+*/
+inline float EncodePanasonicVLog(float Value)
+{
+	const float b = 0.00873f;
+	const float c = 0.241514f;
+	const float d = 0.598206f;
+
+	if( Value < 0.01f)
+	{
+		return 5.6f * Value + 0.125f;
+	}
+	else
+	{
+		return c * FGenericPlatformMath::LogX(10.0f, Value + b) + d;
+	}
+}
+
+/**
+* Decode value with a Panasonic V-Log encoding.
+*
+* @return float decoded value.
+*/
+inline float DecodePanasonicVLog(float Value)
+{
+	const float b = 0.00873f;
+	const float c = 0.241514f;
+	const float d = 0.598206f;
+
+	if (Value < 0.181f)
+	{
+		return (Value - 0.125f) / 5.6f;
+	}
+	else
+	{
+		return FGenericPlatformMath::Pow(10.0f, (Value - d) / c) - b;
+	}
+}
+
+/** Get the encode function that matches the encoding type. */
+CORE_API TFunction<float(float)> GetEncodeFunction(EEncoding Encoding);
+
+/** Get the decode function that matches the encoding type. */
+CORE_API TFunction<float(float)> GetDecodeFunction(EEncoding Encoding);
+
+/** Encode a value based on the specified encoding type. Note: Less optimal due to function pointer call. */
+CORE_API float Encode(EEncoding Encoding, float Value);
+
+/** Decode a value based on the specified encoding type. Note: Less optimal due to function pointer call. */
+CORE_API float Decode(EEncoding Encoding, float Value);
+
+/** Get the encode function that matches the encoding type. */
+CORE_API TFunction<FLinearColor(const FLinearColor&)> GetColorEncodeFunction(EEncoding Encoding);
+
+/** Get the decode function that matches the encoding type. */
+CORE_API TFunction<FLinearColor(const FLinearColor&)> GetColorDecodeFunction(EEncoding Encoding);
+
+/** Encode a color based on the specified encoding type. Note: Less optimal due to function pointer call. */
+CORE_API FLinearColor Encode(EEncoding Encoding, const FLinearColor& Color);
+
+/** Decode a color based on the specified encoding type. Note: Less optimal due to function pointer call. */
+CORE_API FLinearColor Decode(EEncoding Encoding, const FLinearColor& Color);
+
+} } // end namespace UE::Color
