ColorManagement: Fix to missing dll since module move.

#rnx

diff --git a/Engine/Source/Runtime/ColorManagement/Private/ColorManagementModule.cpp b/Engine/Source/Runtime/ColorManagement/Private/ColorManagementModule.cpp
new file mode 100644
--- /dev/null
+++ b/Engine/Source/Runtime/ColorManagement/Private/ColorManagementModule.cpp
@@ -0,0 +1,18 @@
+// Copyright Epic Games, Inc. All Rights Reserved.
+
+#include "Modules/ModuleManager.h"
+
+/**
+ * Implements the ColorManagement module.
+ */
+class FColorManagementModule : public IModuleInterface
+{
+public:
+
+	//~ IModuleInterface interface
+	virtual void StartupModule() override { }
+	virtual void ShutdownModule() override { }
+};
+
+
+IMPLEMENT_MODULE(FColorManagementModule, ColorManagement);
