ColorManagement: Remaining deprecation of headers & color management module, since they were moved to Core at CL 32299544.

#rnx

diff --git a/Engine/Plugins/Compositing/OpenColorIO/Source/OpenColorIO/OpenColorIO.Build.cs b/Engine/Plugins/Compositing/OpenColorIO/Source/OpenColorIO/OpenColorIO.Build.cs
--- a/Engine/Plugins/Compositing/OpenColorIO/Source/OpenColorIO/OpenColorIO.Build.cs
+++ b/Engine/Plugins/Compositing/OpenColorIO/Source/OpenColorIO/OpenColorIO.Build.cs
@@ -25,7 +25,6 @@ namespace UnrealBuildTool.Rules
 					"ImageCore",
 					"Slate",
 					"SlateCore",
-					"ColorManagement"
 				});
 
 			PrivateIncludePathModuleNames.AddRange(
diff --git a/Engine/Plugins/Compositing/OpenColorIO/Source/OpenColorIO/Private/OpenColorIOConfiguration.cpp b/Engine/Plugins/Compositing/OpenColorIO/Source/OpenColorIO/Private/OpenColorIOConfiguration.cpp
--- a/Engine/Plugins/Compositing/OpenColorIO/Source/OpenColorIO/Private/OpenColorIOConfiguration.cpp
+++ b/Engine/Plugins/Compositing/OpenColorIO/Source/OpenColorIO/Private/OpenColorIOConfiguration.cpp
@@ -2,7 +2,7 @@
 
 #include "OpenColorIOConfiguration.h"
 
-#include "ColorSpace.h"
+#include "ColorManagement/ColorSpace.h"
 #include "Containers/SortedMap.h"
 #include "EngineAnalytics.h"
 #include "Engine/VolumeTexture.h"
diff --git a/Engine/Plugins/Experimental/AVCodecs/AVCodecsCore/Source/AVCodecsCoreRHI/AVCodecsCoreRHI.Build.cs b/Engine/Plugins/Experimental/AVCodecs/AVCodecsCore/Source/AVCodecsCoreRHI/AVCodecsCoreRHI.Build.cs
--- a/Engine/Plugins/Experimental/AVCodecs/AVCodecsCore/Source/AVCodecsCoreRHI/AVCodecsCoreRHI.Build.cs
+++ b/Engine/Plugins/Experimental/AVCodecs/AVCodecsCore/Source/AVCodecsCoreRHI/AVCodecsCoreRHI.Build.cs
@@ -16,7 +16,6 @@ public class AVCodecsCoreRHI : ModuleRules
 			"AVCodecsCore",
 			"RHI",
 			"RHICore",
-			"ColorManagement"
 		});
 
 		PublicDependencyModuleNames.AddRange(new string[] {
diff --git a/Engine/Plugins/Experimental/AVCodecs/AVCodecsCore/Source/AVCodecsCoreRHI/Private/Video/Resources/VideoResourceRHI.cpp b/Engine/Plugins/Experimental/AVCodecs/AVCodecsCore/Source/AVCodecsCoreRHI/Private/Video/Resources/VideoResourceRHI.cpp
--- a/Engine/Plugins/Experimental/AVCodecs/AVCodecsCore/Source/AVCodecsCoreRHI/Private/Video/Resources/VideoResourceRHI.cpp
+++ b/Engine/Plugins/Experimental/AVCodecs/AVCodecsCore/Source/AVCodecsCoreRHI/Private/Video/Resources/VideoResourceRHI.cpp
@@ -26,8 +26,8 @@
 #include "RHIStaticStates.h"
 #include "MediaShaders.h"
 #include "PipelineStateCache.h"
-#include "ColorManagementDefines.h"
-#include "ColorSpace.h"
+#include "ColorManagement/ColorManagementDefines.h"
+#include "ColorManagement/ColorSpace.h"
 
 REGISTER_TYPEID(FVideoContextRHI);
 REGISTER_TYPEID(FVideoResourceRHI);
diff --git a/Engine/Plugins/Interchange/Runtime/Source/FactoryNodes/InterchangeFactoryNodes.Build.cs b/Engine/Plugins/Interchange/Runtime/Source/FactoryNodes/InterchangeFactoryNodes.Build.cs
--- a/Engine/Plugins/Interchange/Runtime/Source/FactoryNodes/InterchangeFactoryNodes.Build.cs
+++ b/Engine/Plugins/Interchange/Runtime/Source/FactoryNodes/InterchangeFactoryNodes.Build.cs
@@ -25,7 +25,6 @@ namespace UnrealBuildTool.Rules
 				new string[]
 				{
 					"LevelSequence",
-					"ColorManagement",
 				}
 			);
 		}
diff --git a/Engine/Plugins/Interchange/Runtime/Source/FactoryNodes/Private/InterchangeTextureFactoryNode.cpp b/Engine/Plugins/Interchange/Runtime/Source/FactoryNodes/Private/InterchangeTextureFactoryNode.cpp
--- a/Engine/Plugins/Interchange/Runtime/Source/FactoryNodes/Private/InterchangeTextureFactoryNode.cpp
+++ b/Engine/Plugins/Interchange/Runtime/Source/FactoryNodes/Private/InterchangeTextureFactoryNode.cpp
@@ -2,7 +2,7 @@
 
 #include "InterchangeTextureFactoryNode.h"
 
-#include "ColorSpace.h"
+#include "ColorManagement/ColorSpace.h"
 
 #include UE_INLINE_GENERATED_CPP_BY_NAME(InterchangeTextureFactoryNode)
 
diff --git a/Engine/Plugins/Media/AjaMedia/Source/AjaMediaOutput/AjaMediaOutput.Build.cs b/Engine/Plugins/Media/AjaMedia/Source/AjaMediaOutput/AjaMediaOutput.Build.cs
--- a/Engine/Plugins/Media/AjaMedia/Source/AjaMediaOutput/AjaMediaOutput.Build.cs
+++ b/Engine/Plugins/Media/AjaMedia/Source/AjaMediaOutput/AjaMediaOutput.Build.cs
@@ -18,7 +18,6 @@ namespace UnrealBuildTool.Rules
 				new string[]
 				{
 					"AjaCore",
-					"ColorManagement",
 					"Core",
 					"CoreUObject",
 					"Engine",
diff --git a/Engine/Plugins/Media/BlackmagicMedia/Source/BlackmagicMediaOutput/BlackmagicMediaOutput.Build.cs b/Engine/Plugins/Media/BlackmagicMedia/Source/BlackmagicMediaOutput/BlackmagicMediaOutput.Build.cs
--- a/Engine/Plugins/Media/BlackmagicMedia/Source/BlackmagicMediaOutput/BlackmagicMediaOutput.Build.cs
+++ b/Engine/Plugins/Media/BlackmagicMedia/Source/BlackmagicMediaOutput/BlackmagicMediaOutput.Build.cs
@@ -18,7 +18,6 @@ namespace UnrealBuildTool.Rules
 				new string[]
 				{
 					"BlackmagicCore",
-					"ColorManagement",
 					"Core",
 					"CoreUObject",
 					"Engine",
diff --git a/Engine/Plugins/Media/BlackmagicMedia/Source/BlackmagicMediaOutput/Private/BlackmagicMediaCapture.cpp b/Engine/Plugins/Media/BlackmagicMedia/Source/BlackmagicMediaOutput/Private/BlackmagicMediaCapture.cpp
--- a/Engine/Plugins/Media/BlackmagicMedia/Source/BlackmagicMediaOutput/Private/BlackmagicMediaCapture.cpp
+++ b/Engine/Plugins/Media/BlackmagicMedia/Source/BlackmagicMediaOutput/Private/BlackmagicMediaCapture.cpp
@@ -5,7 +5,7 @@
 #include "BlackmagicLib.h"
 #include "BlackmagicMediaOutput.h"
 #include "BlackmagicMediaOutputModule.h"
-#include "ColorSpace.h"
+#include "ColorManagement/ColorSpace.h"
 #include "GPUTextureTransfer.h"
 #include "GPUTextureTransferModule.h"
 #include "Engine/Engine.h"
diff --git a/Engine/Plugins/Media/ElectraUtil/Source/ElectraSamples/ElectraSamples.Build.cs b/Engine/Plugins/Media/ElectraUtil/Source/ElectraSamples/ElectraSamples.Build.cs
--- a/Engine/Plugins/Media/ElectraUtil/Source/ElectraSamples/ElectraSamples.Build.cs
+++ b/Engine/Plugins/Media/ElectraUtil/Source/ElectraSamples/ElectraSamples.Build.cs
@@ -24,7 +24,6 @@ namespace UnrealBuildTool.Rules
 					"RenderCore",
 					"RHI",
 					"ElectraBase",
-					"ColorManagement",
 				});
 
 			PrivateIncludePathModuleNames.AddRange(
diff --git a/Engine/Plugins/Media/ElectraUtil/Source/ElectraSamples/Public/ElectraTextureSampleUtils.h b/Engine/Plugins/Media/ElectraUtil/Source/ElectraSamples/Public/ElectraTextureSampleUtils.h
--- a/Engine/Plugins/Media/ElectraUtil/Source/ElectraSamples/Public/ElectraTextureSampleUtils.h
+++ b/Engine/Plugins/Media/ElectraUtil/Source/ElectraSamples/Public/ElectraTextureSampleUtils.h
@@ -3,7 +3,7 @@
 #pragma once
 
 #include "HAL/Platform.h"
-#include "ColorManagementDefines.h"
+#include "ColorManagement/ColorManagementDefines.h"
 #include "Misc/AssertionMacros.h"
 
 namespace ElectraColorimetryUtils
diff --git a/Engine/Plugins/Media/MediaIOFramework/Source/MediaIOCore/MediaIOCore.Build.cs b/Engine/Plugins/Media/MediaIOFramework/Source/MediaIOCore/MediaIOCore.Build.cs
--- a/Engine/Plugins/Media/MediaIOFramework/Source/MediaIOCore/MediaIOCore.Build.cs
+++ b/Engine/Plugins/Media/MediaIOFramework/Source/MediaIOCore/MediaIOCore.Build.cs
@@ -13,7 +13,6 @@ namespace UnrealBuildTool.Rules
 				new string[]
 				{
 					"AudioExtensions",
-					"ColorManagement",
 					"Core",
 					"CoreUObject",
 					"Engine",
diff --git a/Engine/Plugins/Media/MediaIOFramework/Source/MediaIOCore/Private/MediaIOCoreDeinterlacer.cpp b/Engine/Plugins/Media/MediaIOFramework/Source/MediaIOCore/Private/MediaIOCoreDeinterlacer.cpp
--- a/Engine/Plugins/Media/MediaIOFramework/Source/MediaIOCore/Private/MediaIOCoreDeinterlacer.cpp
+++ b/Engine/Plugins/Media/MediaIOFramework/Source/MediaIOCore/Private/MediaIOCoreDeinterlacer.cpp
@@ -3,7 +3,7 @@
 #include "MediaIOCoreDeinterlacer.h"
 
 #include "Async/ParallelFor.h"
-#include "ColorManagementDefines.h"
+#include "ColorManagement/ColorManagementDefines.h"
 
 namespace UE::MediaIOCore
 {
diff --git a/Engine/Plugins/Media/MediaIOFramework/Source/MediaIOCore/Public/MediaIOCoreDeinterlacer.h b/Engine/Plugins/Media/MediaIOFramework/Source/MediaIOCore/Public/MediaIOCoreDeinterlacer.h
--- a/Engine/Plugins/Media/MediaIOFramework/Source/MediaIOCore/Public/MediaIOCoreDeinterlacer.h
+++ b/Engine/Plugins/Media/MediaIOFramework/Source/MediaIOCore/Public/MediaIOCoreDeinterlacer.h
@@ -1,7 +1,7 @@
 // Copyright Epic Games, Inc. All Rights Reserved.
 #pragma once
 
-#include "ColorSpace.h"
+#include "ColorManagement/ColorSpace.h"
 #include "Containers/Array.h"
 #include "MediaIOCoreDefinitions.h"
 #include "MediaIOCoreTextureSampleBase.h"
diff --git a/Engine/Plugins/Media/MediaIOFramework/Source/MediaIOCore/Public/MediaIOCoreTextureSampleBase.h b/Engine/Plugins/Media/MediaIOFramework/Source/MediaIOCore/Public/MediaIOCoreTextureSampleBase.h
--- a/Engine/Plugins/Media/MediaIOFramework/Source/MediaIOCore/Public/MediaIOCoreTextureSampleBase.h
+++ b/Engine/Plugins/Media/MediaIOFramework/Source/MediaIOCore/Public/MediaIOCoreTextureSampleBase.h
@@ -8,7 +8,7 @@
 #include "MediaObjectPool.h"
 
 
-#include "ColorSpace.h"
+#include "ColorManagement/ColorSpace.h"
 #include "Misc/FrameRate.h"
 #include "Templates/RefCounting.h"
 
diff --git a/Engine/Plugins/MovieScene/MovieRenderPipeline/Source/MovieRenderPipelineRenderPasses/MovieRenderPipelineRenderPasses.Build.cs b/Engine/Plugins/MovieScene/MovieRenderPipeline/Source/MovieRenderPipelineRenderPasses/MovieRenderPipelineRenderPasses.Build.cs
--- a/Engine/Plugins/MovieScene/MovieRenderPipeline/Source/MovieRenderPipelineRenderPasses/MovieRenderPipelineRenderPasses.Build.cs
+++ b/Engine/Plugins/MovieScene/MovieRenderPipeline/Source/MovieRenderPipelineRenderPasses/MovieRenderPipelineRenderPasses.Build.cs
@@ -22,7 +22,6 @@ public class MovieRenderPipelineRenderPasses : ModuleRules
 				"ImageWrapper",				
 				"CinematicCamera", // For metadata
 				"MovieRenderPipelineSettings", // For settings
-				"ColorManagement",
 				"SlateCore",
 				"ImageCore",
 				"OpenColorIOWrapper",
diff --git a/Engine/Plugins/MovieScene/MovieRenderPipeline/Source/MovieRenderPipelineRenderPasses/Private/MoviePipelineEXROutput.cpp b/Engine/Plugins/MovieScene/MovieRenderPipeline/Source/MovieRenderPipelineRenderPasses/Private/MoviePipelineEXROutput.cpp
--- a/Engine/Plugins/MovieScene/MovieRenderPipeline/Source/MovieRenderPipelineRenderPasses/Private/MoviePipelineEXROutput.cpp
+++ b/Engine/Plugins/MovieScene/MovieRenderPipeline/Source/MovieRenderPipelineRenderPasses/Private/MoviePipelineEXROutput.cpp
@@ -18,7 +18,7 @@
 #include "IOpenExrRTTIModule.h"
 #include "Modules/ModuleManager.h"
 #include "MoviePipelineUtils.h"
-#include "ColorSpace.h"
+#include "ColorManagement/ColorSpace.h"
 #include "HDRHelper.h"
 
 THIRD_PARTY_INCLUDES_START
diff --git a/Engine/Plugins/VirtualProduction/DMX/DMXPixelMapping/Source/DMXPixelMappingRuntime/DMXPixelMappingRuntime.build.cs b/Engine/Plugins/VirtualProduction/DMX/DMXPixelMapping/Source/DMXPixelMappingRuntime/DMXPixelMappingRuntime.build.cs
--- a/Engine/Plugins/VirtualProduction/DMX/DMXPixelMapping/Source/DMXPixelMappingRuntime/DMXPixelMappingRuntime.build.cs
+++ b/Engine/Plugins/VirtualProduction/DMX/DMXPixelMapping/Source/DMXPixelMappingRuntime/DMXPixelMappingRuntime.build.cs
@@ -11,7 +11,6 @@ public class DMXPixelMappingRuntime : ModuleRules
 		PublicDependencyModuleNames.AddRange(
 			new string[]
 			{
-				"ColorManagement",
 				"Core",
 				"CoreUObject",
 				"DMXRuntime",
diff --git a/Engine/Plugins/VirtualProduction/DMX/DMXPixelMapping/Source/DMXPixelMappingRuntime/Internal/ColorSpace/DMXPixelMappingColorSpace_RGBCMY.h b/Engine/Plugins/VirtualProduction/DMX/DMXPixelMapping/Source/DMXPixelMappingRuntime/Internal/ColorSpace/DMXPixelMappingColorSpace_RGBCMY.h
--- a/Engine/Plugins/VirtualProduction/DMX/DMXPixelMapping/Source/DMXPixelMappingRuntime/Internal/ColorSpace/DMXPixelMappingColorSpace_RGBCMY.h
+++ b/Engine/Plugins/VirtualProduction/DMX/DMXPixelMapping/Source/DMXPixelMappingRuntime/Internal/ColorSpace/DMXPixelMappingColorSpace_RGBCMY.h
@@ -5,7 +5,7 @@
 #include "DMXAttribute.h"
 #include "ColorSpace/DMXPixelMappingColorSpace.h"
 
-#include "ColorSpace.h"
+#include "ColorManagement/ColorSpace.h"
 #include "Templates/UniquePtr.h"
 
 #include "DMXPixelMappingColorSpace_RGBCMY.generated.h"
diff --git a/Engine/Plugins/VirtualProduction/DMX/DMXPixelMapping/Source/DMXPixelMappingRuntime/Internal/ColorSpace/DMXPixelMappingColorSpace_XYZ.h b/Engine/Plugins/VirtualProduction/DMX/DMXPixelMapping/Source/DMXPixelMappingRuntime/Internal/ColorSpace/DMXPixelMappingColorSpace_XYZ.h
--- a/Engine/Plugins/VirtualProduction/DMX/DMXPixelMapping/Source/DMXPixelMappingRuntime/Internal/ColorSpace/DMXPixelMappingColorSpace_XYZ.h
+++ b/Engine/Plugins/VirtualProduction/DMX/DMXPixelMapping/Source/DMXPixelMappingRuntime/Internal/ColorSpace/DMXPixelMappingColorSpace_XYZ.h
@@ -5,7 +5,7 @@
 #include "DMXAttribute.h"
 #include "ColorSpace/DMXPixelMappingColorSpace.h"
 
-#include "ColorSpace.h"
+#include "ColorManagement/ColorSpace.h"
 
 #include "DMXPixelMappingColorSpace_XYZ.generated.h"
 
diff --git a/Engine/Plugins/VirtualProduction/DMX/DMXPixelMapping/Source/DMXPixelMappingRuntime/Internal/ColorSpace/DMXPixelMappingColorSpace_xyY.h b/Engine/Plugins/VirtualProduction/DMX/DMXPixelMapping/Source/DMXPixelMappingRuntime/Internal/ColorSpace/DMXPixelMappingColorSpace_xyY.h
--- a/Engine/Plugins/VirtualProduction/DMX/DMXPixelMapping/Source/DMXPixelMappingRuntime/Internal/ColorSpace/DMXPixelMappingColorSpace_xyY.h
+++ b/Engine/Plugins/VirtualProduction/DMX/DMXPixelMapping/Source/DMXPixelMappingRuntime/Internal/ColorSpace/DMXPixelMappingColorSpace_xyY.h
@@ -5,7 +5,7 @@
 #include "DMXAttribute.h"
 #include "ColorSpace/DMXPixelMappingColorSpace.h"
 
-#include "ColorSpace.h"
+#include "ColorManagement/ColorSpace.h"
 
 #include "DMXPixelMappingColorSpace_xyY.generated.h"
 
diff --git a/Engine/Plugins/VirtualProduction/DMX/DMXPixelMapping/Source/DMXPixelMappingRuntime/Private/ColorSpace/DMXPixelMappingColorSpace.cpp b/Engine/Plugins/VirtualProduction/DMX/DMXPixelMapping/Source/DMXPixelMappingRuntime/Private/ColorSpace/DMXPixelMappingColorSpace.cpp
--- a/Engine/Plugins/VirtualProduction/DMX/DMXPixelMapping/Source/DMXPixelMappingRuntime/Private/ColorSpace/DMXPixelMappingColorSpace.cpp
+++ b/Engine/Plugins/VirtualProduction/DMX/DMXPixelMapping/Source/DMXPixelMappingRuntime/Private/ColorSpace/DMXPixelMappingColorSpace.cpp
@@ -2,7 +2,7 @@
 
 #include "ColorSpace/DMXPixelMappingColorSpace.h"
 
-#include "ColorSpace.h"
+#include "ColorManagement/ColorSpace.h"
 #include "CoreMinimal.h"
 #include "UObject/Object.h"
 
diff --git a/Engine/Source/Editor/PixelInspector/PixelInspectorModule.Build.cs b/Engine/Source/Editor/PixelInspector/PixelInspectorModule.Build.cs
--- a/Engine/Source/Editor/PixelInspector/PixelInspectorModule.Build.cs
+++ b/Engine/Source/Editor/PixelInspector/PixelInspectorModule.Build.cs
@@ -23,7 +23,6 @@ public class PixelInspectorModule : ModuleRules
 
 		PrivateDependencyModuleNames.AddRange(
 			new string[] {
-				"ColorManagement",
 				"Renderer",
 			}
 		);
diff --git a/Engine/Source/Editor/PixelInspector/Private/PixelInspectorResult.cpp b/Engine/Source/Editor/PixelInspector/Private/PixelInspectorResult.cpp
--- a/Engine/Source/Editor/PixelInspector/Private/PixelInspectorResult.cpp
+++ b/Engine/Source/Editor/PixelInspector/Private/PixelInspectorResult.cpp
@@ -2,7 +2,7 @@
 
 #include "PixelInspectorResult.h"
 
-#include "ColorSpace.h"
+#include "ColorManagement/ColorSpace.h"
 #include "HAL/PlatformCrt.h"
 #include "Math/Float16.h"
 #include "Math/Float16Color.h"
diff --git a/Engine/Source/Programs/LowLevelTests/FoundationTests/FoundationTests.Build.cs b/Engine/Source/Programs/LowLevelTests/FoundationTests/FoundationTests.Build.cs
--- a/Engine/Source/Programs/LowLevelTests/FoundationTests/FoundationTests.Build.cs
+++ b/Engine/Source/Programs/LowLevelTests/FoundationTests/FoundationTests.Build.cs
@@ -31,7 +31,6 @@ public class FoundationTests : TestModuleRules
 				"CoreUObject",
 				"TelemetryUtils",
 				"AssetRegistry",
-				"ColorManagement",
 				"Serialization",
 			});
 
diff --git a/Engine/Source/Runtime/ColorManagement/ColorManagement.Build.cs b/Engine/Source/Runtime/ColorManagement/ColorManagement.Build.cs
--- a/Engine/Source/Runtime/ColorManagement/ColorManagement.Build.cs
+++ b/Engine/Source/Runtime/ColorManagement/ColorManagement.Build.cs
@@ -1,11 +1,14 @@
 // Copyright Epic Games, Inc. All Rights Reserved.
 
+using System;
 using UnrealBuildTool;
 
 public class ColorManagement : ModuleRules
 {
 	public ColorManagement(ReadOnlyTargetRules Target) : base(Target)
 	{
+		// Note: ColorManagement module is now deprecated in 5.5, as it was moved into Core.
+
 		PublicDependencyModuleNames.AddRange(
 			new string[] {
 				"Core",
diff --git a/Engine/Source/Runtime/ColorManagement/Public/ColorManagementDefines.h b/Engine/Source/Runtime/ColorManagement/Public/ColorManagementDefines.h
--- a/Engine/Source/Runtime/ColorManagement/Public/ColorManagementDefines.h
+++ b/Engine/Source/Runtime/ColorManagement/Public/ColorManagementDefines.h
@@ -2,13 +2,7 @@
 
 #pragma once
 
-//TODO:Uncomment and deprecate this header a in subsequence change
-// 
-//#ifdef _MSC_VER
-//#pragma message(__FILE__"(9): warning: use ColorManagement/ColorManagementDefines.h from the Core module instead of ColorManagementDefines.h")
-//#else
-//#pragma message("#include use ColorManagement/ColorManagementDefines.h from the Core module instead of ColorManagementDefines.h")
-//#endif
+UE_DEPRECATED_HEADER(5.5, "Use ColorManagement/ColorManagementDefines.h from the Core module instead of ColorManagementDefines.h.")
 
 #include "ColorManagement/ColorManagementDefines.h"
 
diff --git a/Engine/Source/Runtime/ColorManagement/Public/ColorSpace.h b/Engine/Source/Runtime/ColorManagement/Public/ColorSpace.h
--- a/Engine/Source/Runtime/ColorManagement/Public/ColorSpace.h
+++ b/Engine/Source/Runtime/ColorManagement/Public/ColorSpace.h
@@ -2,13 +2,7 @@
 
 #pragma once
 
-//TODO:Uncomment and deprecate this header a in subsequence change
-// 
-//#ifdef _MSC_VER
-//#pragma message(__FILE__"(9): warning: use ColorManagement/ColorSpace.h from the Core module instead of ColorSpace.h")
-//#else
-//#pragma message("#include use ColorManagement/ColorSpace.h from the Core module instead of ColorSpace.h")
-//#endif
+UE_DEPRECATED_HEADER(5.5, "Use ColorManagement/ColorSpace.h from the Core module instead of ColorSpace.h.")
 
 #include "ColorManagement/ColorSpace.h"
 
diff --git a/Engine/Source/Runtime/ColorManagement/Public/TransferFunctions.h b/Engine/Source/Runtime/ColorManagement/Public/TransferFunctions.h
--- a/Engine/Source/Runtime/ColorManagement/Public/TransferFunctions.h
+++ b/Engine/Source/Runtime/ColorManagement/Public/TransferFunctions.h
@@ -2,12 +2,6 @@
 
 #pragma once
 
-//TODO: Uncomment and deprecate this header a in subsequence change
-// 
-//#ifdef _MSC_VER
-//#pragma message(__FILE__"(9): warning: use ColorManagement/TransferFunctions.h from the Core module instead of TransferFunctions.h")
-//#else
-//#pragma message("#include use ColorManagement/TransferFunctions.h from the Core module instead of TransferFunctions.h")
-//#endif
+UE_DEPRECATED_HEADER(5.5, "Use ColorManagement/TransferFunctions.h from the Core module instead of TransferFunctions.h.")
 
 #include "ColorManagement/TransferFunctions.h"
diff --git a/Engine/Source/Runtime/Core/Public/ColorManagement/ColorSpace.h b/Engine/Source/Runtime/Core/Public/ColorManagement/ColorSpace.h
--- a/Engine/Source/Runtime/Core/Public/ColorManagement/ColorSpace.h
+++ b/Engine/Source/Runtime/Core/Public/ColorManagement/ColorSpace.h
@@ -2,7 +2,7 @@
 
 #pragma once
 
-#include "ColorManagementDefines.h"
+#include "ColorManagement/ColorManagementDefines.h"
 #include "Containers/StaticArray.h"
 #include "CoreTypes.h"
 #include "Math/Color.h"
diff --git a/Engine/Source/Runtime/Core/Public/ColorManagement/TransferFunctions.h b/Engine/Source/Runtime/Core/Public/ColorManagement/TransferFunctions.h
--- a/Engine/Source/Runtime/Core/Public/ColorManagement/TransferFunctions.h
+++ b/Engine/Source/Runtime/Core/Public/ColorManagement/TransferFunctions.h
@@ -2,7 +2,7 @@
 
 #pragma once
 
-#include "ColorManagementDefines.h"
+#include "ColorManagement/ColorManagementDefines.h"
 #include "CoreTypes.h"
 #include "GenericPlatform/GenericPlatformMath.h"
 #include "Math/Color.h"
diff --git a/Engine/Source/Runtime/Engine/Engine.Build.cs b/Engine/Source/Runtime/Engine/Engine.Build.cs
--- a/Engine/Source/Runtime/Engine/Engine.Build.cs
+++ b/Engine/Source/Runtime/Engine/Engine.Build.cs
@@ -155,7 +155,6 @@ public class Engine : ModuleRules
 				"AudioMixerCore",
 				"IntelISPC",
 				"TraceLog",
-				"ColorManagement",
 				"Icmp",
 				"UniversalObjectLocator",
 				"XmlParser",
diff --git a/Engine/Source/Runtime/Engine/Private/Components/LightComponent.cpp b/Engine/Source/Runtime/Engine/Private/Components/LightComponent.cpp
--- a/Engine/Source/Runtime/Engine/Private/Components/LightComponent.cpp
+++ b/Engine/Source/Runtime/Engine/Private/Components/LightComponent.cpp
@@ -6,7 +6,7 @@
 
 #include "Components/LightComponent.h"
 
-#include "ColorSpace.h"
+#include "ColorManagement/ColorSpace.h"
 #include "Engine/Level.h"
 #include "Engine/MapBuildDataRegistry.h"
 #include "StaticLightingBuildContext.h"
diff --git a/Engine/Source/Runtime/Engine/Private/Components/SkyAtmosphereComponent.cpp b/Engine/Source/Runtime/Engine/Private/Components/SkyAtmosphereComponent.cpp
--- a/Engine/Source/Runtime/Engine/Private/Components/SkyAtmosphereComponent.cpp
+++ b/Engine/Source/Runtime/Engine/Private/Components/SkyAtmosphereComponent.cpp
@@ -17,7 +17,7 @@
 #include "UObject/ConstructorHelpers.h"
 #include "Components/DirectionalLightComponent.h"
 #include "UObject/UE5MainStreamObjectVersion.h"
-#include "ColorSpace.h"
+#include "ColorManagement/ColorSpace.h"
 
 #include UE_INLINE_GENERATED_CPP_BY_NAME(SkyAtmosphereComponent)
 
diff --git a/Engine/Source/Runtime/Engine/Private/Materials/MaterialExpressionHLSL.cpp b/Engine/Source/Runtime/Engine/Private/Materials/MaterialExpressionHLSL.cpp
--- a/Engine/Source/Runtime/Engine/Private/Materials/MaterialExpressionHLSL.cpp
+++ b/Engine/Source/Runtime/Engine/Private/Materials/MaterialExpressionHLSL.cpp
@@ -240,7 +240,7 @@
 #include "MaterialShared.h"
 #include "Misc/MemStackUtility.h"
 #include "RenderUtils.h"
-#include "ColorSpace.h"
+#include "ColorManagement/ColorSpace.h"
 
 bool UMaterialExpression::GenerateHLSLExpression(FMaterialHLSLGenerator& Generator, UE::HLSLTree::FScope& Scope, int32 OutputIndex, UE::HLSLTree::FExpression const*& OutExpression) const
 {
diff --git a/Engine/Source/Runtime/Engine/Private/Materials/MaterialExpressions.cpp b/Engine/Source/Runtime/Engine/Private/Materials/MaterialExpressions.cpp
--- a/Engine/Source/Runtime/Engine/Private/Materials/MaterialExpressions.cpp
+++ b/Engine/Source/Runtime/Engine/Private/Materials/MaterialExpressions.cpp
@@ -42,7 +42,7 @@
 #include "VT/RuntimeVirtualTexture.h"
 #include "SparseVolumeTexture/SparseVolumeTexture.h"
 #include "ProfilingDebugging/LoadTimeTracker.h"
-#include "ColorSpace.h"
+#include "ColorManagement/ColorSpace.h"
 
 #include "Materials/MaterialExpressionAbs.h"
 #include "Materials/MaterialExpressionAbsorptionMediumMaterialOutput.h"
diff --git a/Engine/Source/Runtime/Engine/Private/RendererSettings.cpp b/Engine/Source/Runtime/Engine/Private/RendererSettings.cpp
--- a/Engine/Source/Runtime/Engine/Private/RendererSettings.cpp
+++ b/Engine/Source/Runtime/Engine/Private/RendererSettings.cpp
@@ -2,7 +2,7 @@
 
 #include "Engine/RendererSettings.h"
 #include "GPUSkinVertexFactory.h"
-#include "ColorSpace.h"
+#include "ColorManagement/ColorSpace.h"
 #include "HAL/PlatformFile.h"
 #include "SceneManagement.h"
 #include "Misc/App.h"
diff --git a/Engine/Source/Runtime/Engine/Private/SceneManagement.cpp b/Engine/Source/Runtime/Engine/Private/SceneManagement.cpp
--- a/Engine/Source/Runtime/Engine/Private/SceneManagement.cpp
+++ b/Engine/Source/Runtime/Engine/Private/SceneManagement.cpp
@@ -17,7 +17,7 @@
 #include "TextureResource.h"
 #include "VT/LightmapVirtualTexture.h"
 #include "UnrealEngine.h"
-#include "ColorSpace.h"
+#include "ColorManagement/ColorSpace.h"
 #include "DataDrivenShaderPlatformInfo.h"
 #include "StaticMeshBatch.h"
 #include "PrimitiveUniformShaderParametersBuilder.h"
diff --git a/Engine/Source/Runtime/Engine/Private/ShaderCompiler/ShaderCompiler.cpp b/Engine/Source/Runtime/Engine/Private/ShaderCompiler/ShaderCompiler.cpp
--- a/Engine/Source/Runtime/Engine/Private/ShaderCompiler/ShaderCompiler.cpp
+++ b/Engine/Source/Runtime/Engine/Private/ShaderCompiler/ShaderCompiler.cpp
@@ -59,7 +59,7 @@
 #include "UObject/UObjectIterator.h"
 #include "Math/UnitConversion.h"
 #include "UnrealEngine.h"
-#include "ColorSpace.h"
+#include "ColorManagement/ColorSpace.h"
 
 #if WITH_EDITOR
 #include "UObject/ArchiveCookContext.h"
diff --git a/Engine/Source/Runtime/Engine/Private/Texture.cpp b/Engine/Source/Runtime/Engine/Private/Texture.cpp
--- a/Engine/Source/Runtime/Engine/Private/Texture.cpp
+++ b/Engine/Source/Runtime/Engine/Private/Texture.cpp
@@ -32,7 +32,7 @@
 #include "Compression/OodleDataCompression.h"
 #include "Engine/TextureCube.h"
 #include "Engine/RendererSettings.h"
-#include "ColorSpace.h"
+#include "ColorManagement/ColorSpace.h"
 #include "ImageCoreBP.h"
 #include "ImageCoreUtils.h"
 #include "ImageUtils.h"
diff --git a/Engine/Source/Runtime/Engine/Private/TextureDerivedData.cpp b/Engine/Source/Runtime/Engine/Private/TextureDerivedData.cpp
--- a/Engine/Source/Runtime/Engine/Private/TextureDerivedData.cpp
+++ b/Engine/Source/Runtime/Engine/Private/TextureDerivedData.cpp
@@ -26,7 +26,7 @@
 #if WITH_EDITOR
 
 #include "ChildTextureFormat.h"
-#include "ColorSpace.h"
+#include "ColorManagement/ColorSpace.h"
 #include "Compression/OodleDataCompressionUtil.h"
 #include "DerivedDataCache.h"
 #include "DerivedDataCacheInterface.h"
diff --git a/Engine/Source/Runtime/Engine/Private/TextureDerivedDataBuildUtils.cpp b/Engine/Source/Runtime/Engine/Private/TextureDerivedDataBuildUtils.cpp
--- a/Engine/Source/Runtime/Engine/Private/TextureDerivedDataBuildUtils.cpp
+++ b/Engine/Source/Runtime/Engine/Private/TextureDerivedDataBuildUtils.cpp
@@ -3,7 +3,7 @@
 #include "TextureDerivedDataBuildUtils.h"
 
 #if WITH_EDITOR
-#include "ColorSpace.h"
+#include "ColorManagement/ColorSpace.h"
 #include "DerivedDataBuild.h"
 #include "DerivedDataBuildFunctionRegistry.h"
 #include "DerivedDataSharedString.h"
diff --git a/Engine/Source/Runtime/Engine/Public/Rendering/SkyAtmosphereCommonData.cpp b/Engine/Source/Runtime/Engine/Public/Rendering/SkyAtmosphereCommonData.cpp
--- a/Engine/Source/Runtime/Engine/Public/Rendering/SkyAtmosphereCommonData.cpp
+++ b/Engine/Source/Runtime/Engine/Public/Rendering/SkyAtmosphereCommonData.cpp
@@ -7,7 +7,7 @@
 #include "SkyAtmosphereCommonData.h"
 
 #include "Components/SkyAtmosphereComponent.h"
-#include "ColorSpace.h"
+#include "ColorManagement/ColorSpace.h"
 
 //PRAGMA_DISABLE_OPTIMIZATION
 
diff --git a/Engine/Source/Runtime/ImageCore/ImageCore.Build.cs b/Engine/Source/Runtime/ImageCore/ImageCore.Build.cs
--- a/Engine/Source/Runtime/ImageCore/ImageCore.Build.cs
+++ b/Engine/Source/Runtime/ImageCore/ImageCore.Build.cs
@@ -16,7 +16,6 @@ public class ImageCore : ModuleRules
 			"Engine",
 		});
 
-		PrivateDependencyModuleNames.Add("ColorManagement");
 		PrivateDependencyModuleNames.Add("CoreUObject"); // for TextureDefines.h
 		PrivateDependencyModuleNames.Add("stb_image_resize2");
 
diff --git a/Engine/Source/Runtime/ImageCore/Private/ImageCore.cpp b/Engine/Source/Runtime/ImageCore/Private/ImageCore.cpp
--- a/Engine/Source/Runtime/ImageCore/Private/ImageCore.cpp
+++ b/Engine/Source/Runtime/ImageCore/Private/ImageCore.cpp
@@ -4,8 +4,8 @@
 #include "Modules/ModuleManager.h"
 #include "Async/ParallelFor.h"
 #include "Serialization/CompactBinaryWriter.h"
-#include "TransferFunctions.h"
-#include "ColorSpace.h"
+#include "ColorManagement/TransferFunctions.h"
+#include "ColorManagement/ColorSpace.h"
 #include "ImageParallelFor.h"
 #include "Tasks/Task.h"
 
diff --git a/Engine/Source/Runtime/ImageWrapper/ImageWrapper.Build.cs b/Engine/Source/Runtime/ImageWrapper/ImageWrapper.Build.cs
--- a/Engine/Source/Runtime/ImageWrapper/ImageWrapper.Build.cs
+++ b/Engine/Source/Runtime/ImageWrapper/ImageWrapper.Build.cs
@@ -26,7 +26,6 @@ public class ImageWrapper : ModuleRules
 		PrivateDependencyModuleNames.AddRange(
 			new string[] {
 				"Core",
-				"ColorManagement"
 			}
 		);
 
diff --git a/Engine/Source/Runtime/ImageWrapper/Private/Formats/ExrImageWrapper.cpp b/Engine/Source/Runtime/ImageWrapper/Private/Formats/ExrImageWrapper.cpp
--- a/Engine/Source/Runtime/ImageWrapper/Private/Formats/ExrImageWrapper.cpp
+++ b/Engine/Source/Runtime/ImageWrapper/Private/Formats/ExrImageWrapper.cpp
@@ -4,7 +4,7 @@
 #include "ImageWrapperPrivate.h"
 #include "Math/GuardedInt.h"
 
-#include "ColorSpace.h"
+#include "ColorManagement/ColorSpace.h"
 #include "Containers/StringConv.h"
 #include "HAL/PlatformTime.h"
 #include "Math/Float16.h"
diff --git a/Engine/Source/Runtime/Media/Media.Build.cs b/Engine/Source/Runtime/Media/Media.Build.cs
--- a/Engine/Source/Runtime/Media/Media.Build.cs
+++ b/Engine/Source/Runtime/Media/Media.Build.cs
@@ -13,12 +13,10 @@ namespace UnrealBuildTool.Rules
 				{
 					"Core",
 					"RenderCore",
-					"ColorManagement",
 				});
 
 			PublicIncludePathModuleNames.AddRange(
 				new string[] {
-					"ColorManagement",
 					"RenderCore",
 				});
 		}
diff --git a/Engine/Source/Runtime/Media/Public/IMediaTextureSample.h b/Engine/Source/Runtime/Media/Public/IMediaTextureSample.h
--- a/Engine/Source/Runtime/Media/Public/IMediaTextureSample.h
+++ b/Engine/Source/Runtime/Media/Public/IMediaTextureSample.h
@@ -15,8 +15,8 @@
 #include "Templates/SharedPointer.h"
 
 #include "HDRHelper.h"
-#include "ColorManagementDefines.h"
-#include "ColorSpace.h"
+#include "ColorManagement/ColorManagementDefines.h"
+#include "ColorManagement/ColorSpace.h"
 #include "MediaShaders.h"
 
 #if WITH_ENGINE
diff --git a/Engine/Source/Runtime/MediaAssets/MediaAssets.Build.cs b/Engine/Source/Runtime/MediaAssets/MediaAssets.Build.cs
--- a/Engine/Source/Runtime/MediaAssets/MediaAssets.Build.cs
+++ b/Engine/Source/Runtime/MediaAssets/MediaAssets.Build.cs
@@ -22,7 +22,6 @@ namespace UnrealBuildTool.Rules
 				new string[] {
 					"RenderCore",
 					"RHI",
-					"ColorManagement",
 					"Renderer",
 				});
 
diff --git a/Engine/Source/Runtime/MediaAssets/Private/Misc/MediaTextureResource.h b/Engine/Source/Runtime/MediaAssets/Private/Misc/MediaTextureResource.h
--- a/Engine/Source/Runtime/MediaAssets/Private/Misc/MediaTextureResource.h
+++ b/Engine/Source/Runtime/MediaAssets/Private/Misc/MediaTextureResource.h
@@ -17,7 +17,7 @@
 #include "Async/Async.h"
 #include "RenderingThread.h"
 #include "RendererInterface.h"
-#include "ColorSpace.h"
+#include "ColorManagement/ColorSpace.h"
 
 class FMediaPlayerFacade;
 class IMediaPlayer;
diff --git a/Engine/Source/Runtime/OpenColorIOWrapper/OpenColorIOWrapper.Build.cs b/Engine/Source/Runtime/OpenColorIOWrapper/OpenColorIOWrapper.Build.cs
--- a/Engine/Source/Runtime/OpenColorIOWrapper/OpenColorIOWrapper.Build.cs
+++ b/Engine/Source/Runtime/OpenColorIOWrapper/OpenColorIOWrapper.Build.cs
@@ -8,11 +8,6 @@ namespace UnrealBuildTool.Rules
 		{
 			PCHUsage = PCHUsageMode.UseExplicitOrSharedPCHs;
 
-			PublicIncludePathModuleNames.AddRange(new string[]
-			{
-				"ColorManagement", // for ColorManagementDefines.h
-			});
-
 			PrivateIncludePathModuleNames.AddRange(new string[]
 			{
 				"Engine", // for TextureDefines.h
@@ -35,7 +30,6 @@ namespace UnrealBuildTool.Rules
 				{
 					PrivateDependencyModuleNames.AddRange(new string[]
 					{
-						"ColorManagement",
 						"ImageCore",
 						"OpenColorIOLib",
 					});
diff --git a/Engine/Source/Runtime/OpenColorIOWrapper/Private/OpenColorIOWrapper.cpp b/Engine/Source/Runtime/OpenColorIOWrapper/Private/OpenColorIOWrapper.cpp
--- a/Engine/Source/Runtime/OpenColorIOWrapper/Private/OpenColorIOWrapper.cpp
+++ b/Engine/Source/Runtime/OpenColorIOWrapper/Private/OpenColorIOWrapper.cpp
@@ -8,7 +8,7 @@
 #include "Containers/StringView.h"
 #include "Containers/UnrealString.h"
 #include "Containers/Map.h"
-#include "ColorSpace.h"
+#include "ColorManagement/ColorSpace.h"
 #include "Engine/TextureDefines.h"
 #include "ImageCore.h"
 #include "ImageParallelFor.h"
diff --git a/Engine/Source/Runtime/OpenColorIOWrapper/Private/Tests/OpenColorIOWrapperTests.cpp b/Engine/Source/Runtime/OpenColorIOWrapper/Private/Tests/OpenColorIOWrapperTests.cpp
--- a/Engine/Source/Runtime/OpenColorIOWrapper/Private/Tests/OpenColorIOWrapperTests.cpp
+++ b/Engine/Source/Runtime/OpenColorIOWrapper/Private/Tests/OpenColorIOWrapperTests.cpp
@@ -5,11 +5,11 @@
 
 #if WITH_DEV_AUTOMATION_TESTS && WITH_OCIO
 
-#include "ColorSpace.h"
+#include "ColorManagement/ColorSpace.h"
 #include "Logging/LogMacros.h"
 #include "Math/UnrealMathUtility.h"
 #include "OpenColorIOWrapper.h"
-#include "TransferFunctions.h"
+#include "ColorManagement/TransferFunctions.h"
 
 DEFINE_LOG_CATEGORY_STATIC(LogUnrealOpenColorIOTest, Log, All);
 
diff --git a/Engine/Source/Runtime/OpenColorIOWrapper/Public/OpenColorIOWrapper.h b/Engine/Source/Runtime/OpenColorIOWrapper/Public/OpenColorIOWrapper.h
--- a/Engine/Source/Runtime/OpenColorIOWrapper/Public/OpenColorIOWrapper.h
+++ b/Engine/Source/Runtime/OpenColorIOWrapper/Public/OpenColorIOWrapper.h
@@ -11,7 +11,7 @@
 #include "Templates/PimplPtr.h"
 #include "Templates/UniquePtr.h"
 
-#include "ColorManagementDefines.h"
+#include "ColorManagement/ColorManagementDefines.h"
 
 struct FImageView;
 enum TextureFilter : int;
diff --git a/Engine/Source/Runtime/RenderCore/Private/MediaShaders.cpp b/Engine/Source/Runtime/RenderCore/Private/MediaShaders.cpp
--- a/Engine/Source/Runtime/RenderCore/Private/MediaShaders.cpp
+++ b/Engine/Source/Runtime/RenderCore/Private/MediaShaders.cpp
@@ -5,8 +5,8 @@
 #include "RHIStaticStates.h"
 #include "ShaderParameterUtils.h"
 #include "PipelineStateCache.h"
-#include "ColorManagementDefines.h"
-#include "ColorSpace.h"
+#include "ColorManagement/ColorManagementDefines.h"
+#include "ColorManagement/ColorSpace.h"
 
 FMediaVertexDeclaration::FMediaVertexDeclaration() = default;
 FMediaVertexDeclaration::~FMediaVertexDeclaration() = default;
diff --git a/Engine/Source/Runtime/RenderCore/Private/Shader.cpp b/Engine/Source/Runtime/RenderCore/Private/Shader.cpp
--- a/Engine/Source/Runtime/RenderCore/Private/Shader.cpp
+++ b/Engine/Source/Runtime/RenderCore/Private/Shader.cpp
@@ -24,7 +24,7 @@
 #include "Misc/LargeWorldRenderPosition.h"
 #include "DataDrivenShaderPlatformInfo.h"
 #include "ShaderPlatformCachedIniValue.h"
-#include "ColorSpace.h"
+#include "ColorManagement/ColorSpace.h"
 
 #if WITH_EDITORONLY_DATA
 #include "Interfaces/IShaderFormat.h"
diff --git a/Engine/Source/Runtime/RenderCore/RenderCore.Build.cs b/Engine/Source/Runtime/RenderCore/RenderCore.Build.cs
--- a/Engine/Source/Runtime/RenderCore/RenderCore.Build.cs
+++ b/Engine/Source/Runtime/RenderCore/RenderCore.Build.cs
@@ -37,7 +37,7 @@ public class RenderCore : ModuleRules
 			RuntimeDependencies.Add(Path.Combine(Unreal.EngineDirectory.ToString(), "Extras/GPUDumpViewer/..."), StagedFileType.DebugNonUFS);
 		}
 
-		PrivateDependencyModuleNames.AddRange(new string[] { "Core", "Projects", "ApplicationCore", "TraceLog", "CookOnTheFly", "ColorManagement" });
+		PrivateDependencyModuleNames.AddRange(new string[] { "Core", "Projects", "ApplicationCore", "TraceLog", "CookOnTheFly" });
 
         PrivateIncludePathModuleNames.AddRange(new string[] { "DerivedDataCache" });
 		
